Game Features Editor: Added explicit includes to fix mac compile error
#rnx

diff --git a/Engine/Plugins/Experimental/GameFeatures/Source/GameFeaturesEditor/Private/GameFeatureDataDetailsCustomization.cpp b/Engine/Plugins/Experimental/GameFeatures/Source/GameFeaturesEditor/Private/GameFeatureDataDetailsCustomization.cpp
--- a/Engine/Plugins/Experimental/GameFeatures/Source/GameFeaturesEditor/Private/GameFeatureDataDetailsCustomization.cpp
+++ b/Engine/Plugins/Experimental/GameFeatures/Source/GameFeaturesEditor/Private/GameFeatureDataDetailsCustomization.cpp
@@ -5,6 +5,7 @@
 #include "Widgets/SBoxPanel.h"
 #include "Widgets/Input/SButton.h"
 #include "Widgets/Input/SSegmentedControl.h"
+#include "Widgets/Notifications/SErrorText.h"
 #include "DetailLayoutBuilder.h"
 #include "DetailWidgetRow.h"
 #include "DetailCategoryBuilder.h"
diff --git a/Engine/Plugins/Experimental/GameFeatures/Source/GameFeaturesEditor/Private/GameFeaturesEditorModule.cpp b/Engine/Plugins/Experimental/GameFeatures/Source/GameFeaturesEditor/Private/GameFeaturesEditorModule.cpp
--- a/Engine/Plugins/Experimental/GameFeatures/Source/GameFeaturesEditor/Private/GameFeaturesEditorModule.cpp
+++ b/Engine/Plugins/Experimental/GameFeatures/Source/GameFeaturesEditor/Private/GameFeaturesEditorModule.cpp
@@ -3,6 +3,7 @@
 #include "CoreMinimal.h"
 #include "Modules/ModuleManager.h"
 
+#include "GameFeatureData.h"
 #include "GameFeaturesSubsystemSettings.h"
 #include "Features/IPluginsEditorFeature.h"
 #include "Features/EditorFeatures.h"
@@ -17,6 +18,9 @@
 #include "Engine/AssetManagerSettings.h"
 #include "Widgets/Notifications/SNotificationList.h"
 #include "Framework/Notifications/NotificationManager.h"
+#include "Engine/AssetManager.h"
+#include "HAL/FileManager.h"
+#include "Editor.h"
 
 #define LOCTEXT_NAMESPACE "GameFeatures"
 
