#GameFeatures Split up InitializePluginIniFile so that the base PluginName.ini file can be read during the "Loading" phase before attempting to load content. This file can now contain CoreRedirects

diff --git a/Engine/Plugins/Experimental/GameFeatures/Source/GameFeatures/Private/GameFeatureData.cpp b/Engine/Plugins/Experimental/GameFeatures/Source/GameFeatures/Private/GameFeatureData.cpp
--- a/Engine/Plugins/Experimental/GameFeatures/Source/GameFeatures/Private/GameFeatureData.cpp
+++ b/Engine/Plugins/Experimental/GameFeatures/Source/GameFeatures/Private/GameFeatureData.cpp
@@ -6,6 +6,7 @@
 #include "GameFeaturesSubsystemSettings.h"
 #include "Misc/ConfigCacheIni.h"
 #include "UObject/UObjectHash.h"
+#include "UObject/CoreRedirects.h"
 
 #define LOCTEXT_NAMESPACE "GameFeatures"
 
@@ -54,7 +55,7 @@ EDataValidationResult UGameFeatureData::IsDataValid(TArray<FText>& ValidationErr
 }
 #endif
 
-void UGameFeatureData::InitializePluginIniFile(const FString& PluginInstalledFilename) const
+void UGameFeatureData::InitializeBasePluginIniFile(const FString& PluginInstalledFilename) const
 {
 	const FString PluginName = FPaths::GetBaseFilename(PluginInstalledFilename);
 	const FString PluginConfigDir = FPaths::GetPath(PluginInstalledFilename) / TEXT("Config/");
@@ -78,8 +79,20 @@ void UGameFeatureData::InitializePluginIniFile(const FString& PluginInstalledFil
 	}
 	else
 	{
+		FCoreRedirects::ReadRedirectsFromIni(PluginConfigFilename);
 		ReloadConfigs(PluginConfig);
 	}
+}
+
+void UGameFeatureData::InitializeHierarchicalPluginIniFiles(const FString& PluginInstalledFilename) const
+{
+	const FString PluginName = FPaths::GetBaseFilename(PluginInstalledFilename);
+	const FString PluginConfigDir = FPaths::GetPath(PluginInstalledFilename) / TEXT("Config/");
+	const FString EngineConfigDir = FPaths::EngineConfigDir();
+
+	const bool bIsBaseIniName = false;
+	const bool bForceReloadFromDisk = false;
+	const bool bWriteDestIni = false;
 
 	// @todo: Likely we need to track the diffs this config caused and/or store versions/layers in order to unwind settings during unloading/deactivation
 	TArray<FString> IniNamesToLoad = { TEXT("Input"), TEXT("Game") };
diff --git a/Engine/Plugins/Experimental/GameFeatures/Source/GameFeatures/Private/GameFeaturePluginStateMachine.cpp b/Engine/Plugins/Experimental/GameFeatures/Source/GameFeatures/Private/GameFeaturePluginStateMachine.cpp
--- a/Engine/Plugins/Experimental/GameFeatures/Source/GameFeatures/Private/GameFeaturePluginStateMachine.cpp
+++ b/Engine/Plugins/Experimental/GameFeatures/Source/GameFeatures/Private/GameFeaturePluginStateMachine.cpp
@@ -800,6 +800,8 @@ struct FGameFeaturePluginState_Loading : public FGameFeaturePluginState
 
 		StateStatus.TransitionToState = EGameFeaturePluginState::Loaded;
 
+		StateProperties.GameFeatureData->InitializeBasePluginIniFile(StateProperties.PluginInstalledFilename);
+
 		// AssetManager
 		TSharedPtr<FStreamableHandle> BundleHandle = LoadGameFeatureBundles(StateProperties.GameFeatureData);
 		// @todo make this async. For now we just wait
@@ -913,7 +915,7 @@ struct FGameFeaturePluginState_Activating : public FGameFeaturePluginState
 		check(GEngine);
 		check(StateProperties.GameFeatureData);
 
-		StateProperties.GameFeatureData->InitializePluginIniFile(StateProperties.PluginInstalledFilename);
+		StateProperties.GameFeatureData->InitializeHierarchicalPluginIniFiles(StateProperties.PluginInstalledFilename);
 
 		UGameFeaturesSubsystem::Get().OnGameFeatureActivating(StateProperties.GameFeatureData);
 
diff --git a/Engine/Plugins/Experimental/GameFeatures/Source/GameFeatures/Public/GameFeatureData.h b/Engine/Plugins/Experimental/GameFeatures/Source/GameFeatures/Public/GameFeatureData.h
--- a/Engine/Plugins/Experimental/GameFeatures/Source/GameFeatures/Public/GameFeatureData.h
+++ b/Engine/Plugins/Experimental/GameFeatures/Source/GameFeatures/Public/GameFeatureData.h
@@ -19,8 +19,11 @@ public:
 	/** Method to get where the primary assets should scanned from in the plugin hierarchy */
 	const TArray<FPrimaryAssetTypeInfo>& GetPrimaryAssetTypesToScan() const { return PrimaryAssetTypesToScan; }
 
+	/** Method to process the base ini file for the plugin during loading */
+	void InitializeBasePluginIniFile(const FString& PluginInstalledFilename) const;
+
 	/** Method to process ini files for the plugin during activation */
-	void InitializePluginIniFile(const FString& PluginInstalledFilename) const;
+	void InitializeHierarchicalPluginIniFiles(const FString& PluginInstalledFilename) const;
 
 public:
 	//~UPrimaryDataAsset interface
