#UE4 Add option to keep a plugin registered when unloading

diff --git a/Engine/Plugins/Experimental/GameFeatures/Source/GameFeatures/Private/GameFeaturesSubsystem.cpp b/Engine/Plugins/Experimental/GameFeatures/Source/GameFeatures/Private/GameFeaturesSubsystem.cpp
--- a/Engine/Plugins/Experimental/GameFeatures/Source/GameFeatures/Private/GameFeaturesSubsystem.cpp
+++ b/Engine/Plugins/Experimental/GameFeatures/Source/GameFeatures/Private/GameFeaturesSubsystem.cpp
@@ -431,19 +431,20 @@ void UGameFeaturesSubsystem::DeactivateGameFeaturePlugin(const FString& PluginUR
 	}
 }
 
-void UGameFeaturesSubsystem::UnloadGameFeaturePlugin(const FString& PluginURL)
+void UGameFeaturesSubsystem::UnloadGameFeaturePlugin(const FString& PluginURL, bool bKeepRegistered)
 {
 	FGameFeaturePluginUnloadComplete Callback = FGameFeaturePluginUnloadComplete();
-	UnloadGameFeaturePlugin(PluginURL, Callback);
+	UnloadGameFeaturePlugin(PluginURL, Callback, bKeepRegistered);
 }
 
-void UGameFeaturesSubsystem::UnloadGameFeaturePlugin(const FString& PluginURL, const FGameFeaturePluginUnloadComplete& CompleteDelegate)
+void UGameFeaturesSubsystem::UnloadGameFeaturePlugin(const FString& PluginURL, const FGameFeaturePluginUnloadComplete& CompleteDelegate, bool bKeepRegistered)
 {
 	if (UGameFeaturePluginStateMachine* StateMachine = GetGameFeaturePluginStateMachine(PluginURL, false))
 	{
-		if (StateMachine->GetCurrentState() > EGameFeaturePluginState::Installed)
+		EGameFeaturePluginState::Type DestinationState = bKeepRegistered ? EGameFeaturePluginState::Registered : EGameFeaturePluginState::Installed;
+		if (StateMachine->GetCurrentState() > DestinationState)
 		{
-			StateMachine->SetDestinationState(EGameFeaturePluginState::Installed, FGameFeatureStateTransitionComplete::CreateUObject(this, &ThisClass::UnloadGameFeaturePluginComplete, CompleteDelegate));
+			StateMachine->SetDestinationState(DestinationState, FGameFeatureStateTransitionComplete::CreateUObject(this, &ThisClass::UnloadGameFeaturePluginComplete, CompleteDelegate));
 		}
 	}
 }
diff --git a/Engine/Plugins/Experimental/GameFeatures/Source/GameFeatures/Public/GameFeaturesSubsystem.h b/Engine/Plugins/Experimental/GameFeatures/Source/GameFeatures/Public/GameFeaturesSubsystem.h
--- a/Engine/Plugins/Experimental/GameFeatures/Source/GameFeatures/Public/GameFeaturesSubsystem.h
+++ b/Engine/Plugins/Experimental/GameFeatures/Source/GameFeatures/Public/GameFeaturesSubsystem.h
@@ -142,8 +142,8 @@ public:
 	void DeactivateGameFeaturePlugin(const FString& PluginURL, const FGameFeaturePluginDeactivateComplete& CompleteDelegate);
 
 	/** Unloads the specified game feature plugin. */
-	void UnloadGameFeaturePlugin(const FString& PluginURL);
-	void UnloadGameFeaturePlugin(const FString& PluginURL, const FGameFeaturePluginUnloadComplete& CompleteDelegate);
+	void UnloadGameFeaturePlugin(const FString& PluginURL, bool bKeepRegistered = false);
+	void UnloadGameFeaturePlugin(const FString& PluginURL, const FGameFeaturePluginUnloadComplete& CompleteDelegate, bool bKeepRegistered = false);
 
 	/** Uninstall the specified game feature plugin. Will remove the game feature plugin from the device if it was downloaded */
 	void UninstallGameFeaturePlugin(const FString& PluginURL);
