#UE4 BuiltIn GameFeaturePlugins initial state can now be overridden in game feature policies. Changed up storage to an enum.

diff --git a/Engine/Plugins/Experimental/GameFeatures/Source/GameFeatures/Private/GameFeaturesProjectPolicies.cpp b/Engine/Plugins/Experimental/GameFeatures/Source/GameFeatures/Private/GameFeaturesProjectPolicies.cpp
--- a/Engine/Plugins/Experimental/GameFeatures/Source/GameFeatures/Private/GameFeaturesProjectPolicies.cpp
+++ b/Engine/Plugins/Experimental/GameFeatures/Source/GameFeatures/Private/GameFeaturesProjectPolicies.cpp
@@ -8,7 +8,7 @@ void UDefaultGameFeaturesProjectPolicies::InitGameFeatureManager()
 {
 	UE_LOG(LogGameFeatures, Log, TEXT("Scanning for built-in game feature plugins"));
 
-	auto AdditionalFilter = [&](const FString& PluginFilename, const FGameFeaturePluginDetails& PluginDetails) -> bool
+	auto AdditionalFilter = [&](const FString& PluginFilename, const FGameFeaturePluginDetails& PluginDetails, FBuiltInGameFeaturePluginBehaviorOptions& OutOptions) -> bool
 	{
 		return true;
 	};
diff --git a/Engine/Plugins/Experimental/GameFeatures/Source/GameFeatures/Private/GameFeaturesSubsystem.cpp b/Engine/Plugins/Experimental/GameFeatures/Source/GameFeatures/Private/GameFeaturesSubsystem.cpp
--- a/Engine/Plugins/Experimental/GameFeatures/Source/GameFeatures/Private/GameFeaturesSubsystem.cpp
+++ b/Engine/Plugins/Experimental/GameFeatures/Source/GameFeatures/Private/GameFeaturesSubsystem.cpp
@@ -441,7 +441,7 @@ void UGameFeaturesSubsystem::UnloadGameFeaturePlugin(const FString& PluginURL, c
 {
 	if (UGameFeaturePluginStateMachine* StateMachine = GetGameFeaturePluginStateMachine(PluginURL, false))
 	{
-		EGameFeaturePluginState::Type DestinationState = bKeepRegistered ? EGameFeaturePluginState::Registered : EGameFeaturePluginState::Installed;
+		EGameFeaturePluginState DestinationState = bKeepRegistered ? EGameFeaturePluginState::Registered : EGameFeaturePluginState::Installed;
 		if (StateMachine->GetCurrentState() > DestinationState)
 		{
 			StateMachine->SetDestinationState(DestinationState, FGameFeatureStateTransitionComplete::CreateUObject(this, &ThisClass::UnloadGameFeaturePluginComplete, CompleteDelegate));
@@ -480,24 +480,34 @@ void UGameFeaturesSubsystem::LoadBuiltInGameFeaturePlugin(const TSharedRef<IPlug
 			FGameFeaturePluginDetails PluginDetails;
 			if (GetGameFeaturePluginDetails(PluginDescriptorFilename, PluginDetails))
 			{
-				bool bShouldProcess = AdditionalFilter(PluginDescriptorFilename, PluginDetails);
+				FBuiltInGameFeaturePluginBehaviorOptions BehaviorOptions;
+				bool bShouldProcess = AdditionalFilter(PluginDescriptorFilename, PluginDetails, BehaviorOptions);
 
 				if (bShouldProcess)
 				{
 					UGameFeaturePluginStateMachine* StateMachine = GetGameFeaturePluginStateMachine(PluginURL, true);
 
-					EGameFeaturePluginState DestinationState = EGameFeaturePluginState::Installed;
-					if (PluginDetails.bBuiltInAutoRegister)
+					EBuiltInAutoState InitialAutoState = BehaviorOptions.AutoStateOverride != EBuiltInAutoState::Invalid ? BehaviorOptions.AutoStateOverride : PluginDetails.BuiltInAutoState;
+						
+					EGameFeaturePluginState DestinationState = EGameFeaturePluginState::UnknownStatus;
+					switch (InitialAutoState)
 					{
+					case EBuiltInAutoState::Installed:
+						DestinationState = EGameFeaturePluginState::Installed;
+						break;
+					case EBuiltInAutoState::Registered:
 						DestinationState = EGameFeaturePluginState::Registered;
-						if (PluginDetails.bBuiltInAutoLoad)
-						{
-							DestinationState = EGameFeaturePluginState::Loaded;
-							if (PluginDetails.bBuiltInAutoActivate)
-							{
-								DestinationState = EGameFeaturePluginState::Active;
-							}
-						}
+						break;
+					case EBuiltInAutoState::Loaded:
+						DestinationState = EGameFeaturePluginState::Loaded;
+						break;
+					case EBuiltInAutoState::Active:
+						DestinationState = EGameFeaturePluginState::Active;
+						break;
+
+					default:
+						ensureMsgf(false, TEXT("BuilIn Game Feature Plugin %s configured for invalid auto state %d. Setting to UnknownStatus."), (uint8)InitialAutoState);
+						break;
 					}
 
 					if (StateMachine->GetCurrentState() >= DestinationState)
@@ -608,16 +618,29 @@ bool UGameFeaturesSubsystem::GetGameFeaturePluginDetails(const FString& PluginDe
 	ObjectPtr->TryGetBoolField(TEXT("Hotfixable"), OutPluginDetails.bHotfixable);
 
 	// BuiltInAutoRegister. Default to true. If this is a built in plugin, should it be registered automatically (set to false if you intent to load late with LoadAndActivateGameFeaturePlugin)
-	OutPluginDetails.bBuiltInAutoRegister = true;
-	ObjectPtr->TryGetBoolField(TEXT("BuiltInAutoRegister"), OutPluginDetails.bBuiltInAutoRegister);
+	bool bBuiltInAutoRegister = true;
+	ObjectPtr->TryGetBoolField(TEXT("BuiltInAutoRegister"), bBuiltInAutoRegister);
 
 	// BuiltInAutoLoad. Default to true. If this is a built in plugin, should it be loaded automatically (set to false if you intent to load late with LoadAndActivateGameFeaturePlugin)
-	OutPluginDetails.bBuiltInAutoLoad = true;
-	ObjectPtr->TryGetBoolField(TEXT("BuiltInAutoLoad"), OutPluginDetails.bBuiltInAutoLoad);
+	bool bBuiltInAutoLoad = true;
+	ObjectPtr->TryGetBoolField(TEXT("BuiltInAutoLoad"), bBuiltInAutoLoad);
 
 	// The cooker will need to activate the plugin so that assets can be scanned properly
-	OutPluginDetails.bBuiltInAutoActivate = true;
-	ObjectPtr->TryGetBoolField(TEXT("BuiltInAutoActivate"), OutPluginDetails.bBuiltInAutoActivate);
+	bool bBuiltInAutoActivate = true;
+	ObjectPtr->TryGetBoolField(TEXT("BuiltInAutoActivate"), bBuiltInAutoActivate);
+
+	if (bBuiltInAutoRegister)
+	{
+		OutPluginDetails.BuiltInAutoState = EBuiltInAutoState::Registered;
+		if (bBuiltInAutoLoad)
+		{
+			OutPluginDetails.BuiltInAutoState = EBuiltInAutoState::Loaded;
+			if (bBuiltInAutoActivate)
+			{
+				OutPluginDetails.BuiltInAutoState = EBuiltInAutoState::Active;
+			}
+		}
+	}
 
 	// Read any additional metadata the policy might want to consume (e.g., a release version number)
 	for (const FString& ExtraKey : GetDefault<UGameFeaturesSubsystemSettings>()->AdditionalPluginMetadataKeys)
diff --git a/Engine/Plugins/Experimental/GameFeatures/Source/GameFeatures/Public/GameFeaturesSubsystem.h b/Engine/Plugins/Experimental/GameFeatures/Source/GameFeatures/Public/GameFeaturesSubsystem.h
--- a/Engine/Plugins/Experimental/GameFeatures/Source/GameFeatures/Public/GameFeaturesSubsystem.h
+++ b/Engine/Plugins/Experimental/GameFeatures/Source/GameFeatures/Public/GameFeaturesSubsystem.h
@@ -59,25 +59,32 @@ DECLARE_MULTICAST_DELEGATE_TwoParams(FGameFeaturePluginLoadCompleteDataReady, co
 /** Notification that a game feature plugin load has deactivated successfully and feeds back the GameFeatureData that was being used*/
 DECLARE_MULTICAST_DELEGATE_TwoParams(FGameFeaturePluginDeativated, const FString& /*Name*/, const UGameFeatureData* /*Data*/);
 
+enum class EBuiltInAutoState : uint8
+{
+	Invalid,
+	Installed,
+	Registered,
+	Loaded,
+	Active
+};
 
 struct FGameFeaturePluginDetails
 {
 	TArray<FString> PluginDependencies;
 	TMap<FString, FString> AdditionalMetadata;
 	bool bHotfixable;
-	bool bBuiltInAutoRegister;
-	bool bBuiltInAutoLoad;
-	bool bBuiltInAutoActivate;
+	EBuiltInAutoState BuiltInAutoState;
 
 	FGameFeaturePluginDetails()
 		: bHotfixable(false)
-		, bBuiltInAutoRegister(false)
-		, bBuiltInAutoLoad(false)
-		, bBuiltInAutoActivate(false)
+		, BuiltInAutoState(EBuiltInAutoState::Installed)
 	{}
 };
 
-
+struct FBuiltInGameFeaturePluginBehaviorOptions
+{
+	EBuiltInAutoState AutoStateOverride = EBuiltInAutoState::Invalid;
+};
 
 
 
@@ -162,7 +169,7 @@ public:
 		return *CastChecked<T>(GameSpecificPolicies, ECastCheckedType::NullChecked);
 	}
 
-	typedef TFunctionRef<bool(const FString& PluginFilename, const FGameFeaturePluginDetails& Details)> FBuiltInPluginAdditionalFilters;
+	typedef TFunctionRef<bool(const FString& PluginFilename, const FGameFeaturePluginDetails& Details, FBuiltInGameFeaturePluginBehaviorOptions& OutOptions)> FBuiltInPluginAdditionalFilters;
 
 	/** Loads a built-in game feature plugin if it passes the specified filter */
 	void LoadBuiltInGameFeaturePlugin(const TSharedRef<IPlugin>& Plugin, FBuiltInPluginAdditionalFilters AdditionalFilter);
diff --git a/Engine/Plugins/Experimental/GameFeatures/Source/GameFeaturesEditor/Private/GameFeaturesEditorModule.cpp b/Engine/Plugins/Experimental/GameFeatures/Source/GameFeaturesEditor/Private/GameFeaturesEditorModule.cpp
--- a/Engine/Plugins/Experimental/GameFeatures/Source/GameFeaturesEditor/Private/GameFeaturesEditorModule.cpp
+++ b/Engine/Plugins/Experimental/GameFeatures/Source/GameFeaturesEditor/Private/GameFeaturesEditorModule.cpp
@@ -76,7 +76,7 @@ struct FGameFeaturePluginTemplateDescription : public FPluginTemplateDescription
 		UObject* NewAsset = AssetToolsModule.Get().CreateAsset(NewPlugin->GetName(), NewPlugin->GetMountedAssetPath(), DefaultGameFeatureDataClass, /*Factory=*/ nullptr);
 
 		// Activate the new game feature plugin
-		auto AdditionalFilter = [](const FString&, const FGameFeaturePluginDetails&) -> bool { return true; };
+		auto AdditionalFilter = [](const FString&, const FGameFeaturePluginDetails&, FBuiltInGameFeaturePluginBehaviorOptions&) -> bool { return true; };
 		UGameFeaturesSubsystem::Get().LoadBuiltInGameFeaturePlugin(NewPlugin.ToSharedRef(), AdditionalFilter);
 
 		// Edit the new game feature data
