Made it possible to go from Active or Loaded to Registered via the switcher in the details panel (Installed->Registered is still not possible)

diff --git a/Engine/Plugins/Experimental/GameFeatures/Source/GameFeaturesEditor/Private/GameFeatureDataDetailsCustomization.cpp b/Engine/Plugins/Experimental/GameFeatures/Source/GameFeaturesEditor/Private/GameFeatureDataDetailsCustomization.cpp
--- a/Engine/Plugins/Experimental/GameFeatures/Source/GameFeaturesEditor/Private/GameFeatureDataDetailsCustomization.cpp
+++ b/Engine/Plugins/Experimental/GameFeatures/Source/GameFeaturesEditor/Private/GameFeatureDataDetailsCustomization.cpp
@@ -163,7 +163,7 @@ void FGameFeatureDataDetailsCustomization::ChangeDesiredState(EGameFeaturePlugin
 	{
 		if (CurrentState >= EGameFeaturePluginState::Loaded)
 		{
-			Subsystem.UnloadGameFeaturePlugin(PluginURL);
+			Subsystem.UnloadGameFeaturePlugin(PluginURL, FGameFeaturePluginDeactivateComplete::CreateStatic(&FGameFeatureDataDetailsCustomization::OnOperationCompletedOrFailed, WeakThisPtr), /*bKeepRegistered=*/ true);
 		}
 		else
 		{
@@ -174,7 +174,7 @@ void FGameFeatureDataDetailsCustomization::ChangeDesiredState(EGameFeaturePlugin
 	{
 		//@TODO: No public transition from something greater than Installed to Installed is exposed yet
 		//@TODO: Do we need to support unregistering?  If not, should remove this button
-		Subsystem.UnloadGameFeaturePlugin(PluginURL);
+		Subsystem.UnloadGameFeaturePlugin(PluginURL, FGameFeaturePluginDeactivateComplete::CreateStatic(&FGameFeatureDataDetailsCustomization::OnOperationCompletedOrFailed, WeakThisPtr), /*bKeepRegistered=*/ false);
 	}
 }
 
