Added an explicit include since forward declaration for a type referenced in TSubclassOf doesn't seem to reliably work anymore...

diff --git a/Engine/Plugins/Experimental/GameFeatures/Source/GameFeatures/Public/GameFeaturesSubsystemSettings.h b/Engine/Plugins/Experimental/GameFeatures/Source/GameFeatures/Public/GameFeaturesSubsystemSettings.h
--- a/Engine/Plugins/Experimental/GameFeatures/Source/GameFeatures/Public/GameFeaturesSubsystemSettings.h
+++ b/Engine/Plugins/Experimental/GameFeatures/Source/GameFeatures/Public/GameFeaturesSubsystemSettings.h
@@ -4,11 +4,9 @@
 
 #include "CoreMinimal.h"
 #include "Engine/DeveloperSettings.h"
-
+#include "GameFeatureData.h"
 #include "GameFeaturesSubsystemSettings.generated.h"
 
-class UGameFeatureData;
-
 /** Settings for the Game Features framework */
 UCLASS(config=Game, defaultconfig, notplaceable, meta = (DisplayName = "Game Features"))
 class GAMEFEATURES_API UGameFeaturesSubsystemSettings : public UDeveloperSettings
