Made the 'Add entry to PrimaryAssetTypesToScan?' action update the asset manager so an editor reboot isn't necessary between using it and adding a game feature plugin to a project

diff --git a/Engine/Plugins/Experimental/GameFeatures/Source/GameFeaturesEditor/Private/GameFeaturesEditorModule.cpp b/Engine/Plugins/Experimental/GameFeatures/Source/GameFeaturesEditor/Private/GameFeaturesEditorModule.cpp
--- a/Engine/Plugins/Experimental/GameFeatures/Source/GameFeaturesEditor/Private/GameFeaturesEditorModule.cpp
+++ b/Engine/Plugins/Experimental/GameFeatures/Source/GameFeaturesEditor/Private/GameFeaturesEditorModule.cpp
@@ -245,6 +245,8 @@ class FGameFeaturesEditorModule : public FDefaultModuleImpl
 
  			Settings->PostEditChange();
 			Settings->UpdateDefaultConfigFile();
+
+			UAssetManager::Get().ReinitializeFromConfig();
 		}
 
 		// Show a message that the file was checked out/updated and must be submitted
@@ -261,8 +263,8 @@ class FGameFeaturesEditorModule : public FDefaultModuleImpl
 		if (GameDataRules.IsDefault())
 		{
 			FMessageLog("LoadErrors").Error()
-				->AddToken(FTextToken::Create(FText::Format(NSLOCTEXT("GameFeatures", "MissingRuleForGameFeatureData", "Asset Manager settings do not include a rule for assets of type {0}, which is required for game feature plugins to function.  Add a rule and restart the editor."), FText::FromName(UGameFeatureData::StaticClass()->GetFName()))))
-				->AddToken(FActionToken::Create(NSLOCTEXT("GameFeatures", "AddRuleForGameFeatureData", "Add placeholder rule"), FText(),
+				->AddToken(FTextToken::Create(FText::Format(NSLOCTEXT("GameFeatures", "MissingRuleForGameFeatureData", "Asset Manager settings do not include an entry for assets of type {0}, which is required for game feature plugins to function."), FText::FromName(UGameFeatureData::StaticClass()->GetFName()))))
+				->AddToken(FActionToken::Create(NSLOCTEXT("GameFeatures", "AddRuleForGameFeatureData", "Add entry to PrimaryAssetTypesToScan?"), FText(),
 					FOnActionTokenExecuted::CreateRaw(this, &FGameFeaturesEditorModule::AddDefaultGameDataRule), true));
 		}
 	}
