Fix it so game feature plugins will correctly scan full asset manager paths (which the editor UI sets) as well as plugin-relative ones
Add helper function to convert between plugin relative and absolute paths, can move it elsewhere

diff --git a/Engine/Plugins/Experimental/GameFeatures/Source/GameFeatures/Private/GameFeaturesSubsystem.cpp b/Engine/Plugins/Experimental/GameFeatures/Source/GameFeatures/Private/GameFeaturesSubsystem.cpp
--- a/Engine/Plugins/Experimental/GameFeatures/Source/GameFeatures/Private/GameFeaturesSubsystem.cpp
+++ b/Engine/Plugins/Experimental/GameFeatures/Source/GameFeatures/Private/GameFeaturesSubsystem.cpp
@@ -200,8 +200,8 @@ void UGameFeaturesSubsystem::AddGameFeatureToAssetManager(const UGameFeatureData
 	UAssetManager& LocalAssetManager = UAssetManager::Get();
 
 	const FString GameFeaturePath = GameFeatureToAdd->GetOutermost()->GetName();
-	FString PluginName;
-	if (ensureMsgf(UAssetManager::GetContentRootPathFromPackageName(GameFeaturePath, PluginName), TEXT("Must be a valid package path with a root. GameFeaturePath: %s"), *GameFeaturePath))
+	FString PluginRootPath;
+	if (ensureMsgf(UAssetManager::GetContentRootPathFromPackageName(GameFeaturePath, PluginRootPath), TEXT("Must be a valid package path with a root. GameFeaturePath: %s"), *GameFeaturePath))
 	{
 		LocalAssetManager.StartBulkScanning();
 
@@ -215,7 +215,8 @@ void UGameFeaturesSubsystem::AddGameFeatureToAssetManager(const UGameFeatureData
 
 			for (FString& Path : TypeInfo.AssetScanPaths)
 			{
-				Path = PluginName + Path;
+				// Convert plugin-relative paths to full package paths
+				FixPluginPackagePath(Path, PluginRootPath, false);
 			}
 
 			FPrimaryAssetTypeInfo ExistingAssetTypeInfo;
@@ -564,6 +565,22 @@ FString UGameFeaturesSubsystem::GetPluginFilenameFromPluginURL(const FString& Pl
 	return PluginFilename;
 }
 
+void UGameFeaturesSubsystem::FixPluginPackagePath(FString& PathToFix, const FString& PluginRootPath, bool bMakeRelativeToPluginRoot)
+{
+	if (bMakeRelativeToPluginRoot)
+	{
+		// This only modifies paths starting with the root
+		PathToFix.RemoveFromStart(PluginRootPath);
+	}
+	else
+	{
+		if (!FPackageName::IsValidLongPackageName(PathToFix))
+		{
+			PathToFix = PluginRootPath / PathToFix;
+		}
+	}
+}
+
 void UGameFeaturesSubsystem::GetLoadedGameFeaturePluginFilenamesForCooking(TArray<FString>& OutLoadedPluginFilenames) const
 {
 	for (auto StateMachineIt = GameFeaturePluginStateMachines.CreateConstIterator(); StateMachineIt; ++StateMachineIt)
diff --git a/Engine/Plugins/Experimental/GameFeatures/Source/GameFeatures/Public/GameFeaturesSubsystem.h b/Engine/Plugins/Experimental/GameFeatures/Source/GameFeatures/Public/GameFeaturesSubsystem.h
--- a/Engine/Plugins/Experimental/GameFeatures/Source/GameFeatures/Public/GameFeaturesSubsystem.h
+++ b/Engine/Plugins/Experimental/GameFeatures/Source/GameFeatures/Public/GameFeaturesSubsystem.h
@@ -162,6 +162,9 @@ public:
 	/** Get the plugin path from the plugin URL */
 	FString GetPluginFilenameFromPluginURL(const FString& PluginURL);
 
+	/** Fixes a package path/directory to either be relative to plugin root or not. Paths relative to different roots will not be modified */
+	static void FixPluginPackagePath(FString& PathToFix, const FString& PluginRootPath, bool bMakeRelativeToPluginRoot);
+
 	/** Returns the game-specific policy for managing game feature plugins */
 	template <typename T = UGameFeaturesProjectPolicies>
 	T& GetPolicy() const
