Plugin editing refactors, game feature UX improvements
- Added a new plugin key BuiltInInitialFeatureState for game feature plugins to replace the three separate booleans (the old bools are still supported if this key is missing)
- Added a link from the game feature data editor to get to the plugin editor for the associated plugin
- Added display of initial plugin state to the game feature data editor
- Added a so-so mechanism to extend the plugin descriptor editor from other plugins
- Used this to allow editing the initial state for game feature plugins

diff --git a/Engine/Plugins/Editor/PluginBrowser/Source/PluginBrowser/PluginBrowser.Build.cs b/Engine/Plugins/Editor/PluginBrowser/Source/PluginBrowser/PluginBrowser.Build.cs
--- a/Engine/Plugins/Editor/PluginBrowser/Source/PluginBrowser/PluginBrowser.Build.cs
+++ b/Engine/Plugins/Editor/PluginBrowser/Source/PluginBrowser/PluginBrowser.Build.cs
@@ -31,6 +31,7 @@ namespace UnrealBuildTool.Rules
 					"MainFrame",
                     "UATHelper",
                     "AssetTools",
+					"Json"
 				}
 			);
 
diff --git a/Engine/Plugins/Editor/PluginBrowser/Source/PluginBrowser/Private/PluginBrowserModule.cpp b/Engine/Plugins/Editor/PluginBrowser/Source/PluginBrowser/Private/PluginBrowserModule.cpp
--- a/Engine/Plugins/Editor/PluginBrowser/Source/PluginBrowser/Private/PluginBrowserModule.cpp
+++ b/Engine/Plugins/Editor/PluginBrowser/Source/PluginBrowser/Private/PluginBrowserModule.cpp
@@ -16,6 +16,7 @@
 #include "Framework/Notifications/NotificationManager.h"
 #include "Misc/ConfigCacheIni.h"
 #include "Interfaces/IPluginManager.h"
+#include "PluginDescriptorEditor.h"
 
 #define LOCTEXT_NAMESPACE "PluginsEditor"
 
@@ -108,6 +109,24 @@ void FPluginBrowserModule::UnregisterPluginTemplate(TSharedRef<FPluginTemplateDe
 	AddedPluginTemplates.RemoveSingle(Template);
 }
 
+FPluginEditorExtensionHandle FPluginBrowserModule::RegisterPluginEditorExtension(FOnPluginBeingEdited Extension)
+{
+	++EditorExtensionCounter;
+	FPluginEditorExtensionHandle Result = EditorExtensionCounter;
+	CustomizePluginEditingDelegates.Add(MakeTuple(Extension, Result));
+	return Result;
+}
+
+void FPluginBrowserModule::UnregisterPluginEditorExtension(FPluginEditorExtensionHandle ExtensionHandle)
+{
+	CustomizePluginEditingDelegates.RemoveAll([=](TPair<FOnPluginBeingEdited, FPluginEditorExtensionHandle>& Value) { return Value.Value == ExtensionHandle; });
+}
+
+void FPluginBrowserModule::OpenPluginEditor(TSharedRef<IPlugin> PluginToEdit, TSharedPtr<SWidget> ParentWidget, FSimpleDelegate OnEditCommitted)
+{
+	FPluginDescriptorEditor::OpenEditorWindow(PluginToEdit, ParentWidget, OnEditCommitted);
+}
+
 void FPluginBrowserModule::SetPluginPendingEnableState(const FString& PluginName, bool bCurrentlyEnabled, bool bPendingEnabled)
 {
 	if (bCurrentlyEnabled == bPendingEnabled)
diff --git a/Engine/Plugins/Editor/PluginBrowser/Source/PluginBrowser/Private/PluginBrowserModule.h b/Engine/Plugins/Editor/PluginBrowser/Source/PluginBrowser/Private/PluginBrowserModule.h
--- a/Engine/Plugins/Editor/PluginBrowser/Source/PluginBrowser/Private/PluginBrowserModule.h
+++ b/Engine/Plugins/Editor/PluginBrowser/Source/PluginBrowser/Private/PluginBrowserModule.h
@@ -28,6 +28,9 @@ public:
 	/** IPluginsEditorFeature implementation */
 	virtual void RegisterPluginTemplate(TSharedRef<FPluginTemplateDescription> Template) override;
 	virtual void UnregisterPluginTemplate(TSharedRef<FPluginTemplateDescription> Template) override;
+	virtual FPluginEditorExtensionHandle RegisterPluginEditorExtension(FOnPluginBeingEdited Extension) override;
+	virtual void UnregisterPluginEditorExtension(FPluginEditorExtensionHandle ExtensionHandle) override;
+	virtual void OpenPluginEditor(TSharedRef<IPlugin> PluginToEdit, TSharedPtr<SWidget> ParentWidget, FSimpleDelegate OnEditCommitted) override;
 
 	/** Gets a delegate so that you can register/unregister to receive callbacks when plugins are created */
 	FOnNewPluginCreated& OnNewPluginCreated() {return NewPluginCreatedDelegate;}
@@ -73,6 +76,8 @@ public:
 
 	const TArray<TSharedRef<FPluginTemplateDescription>>& GetAddedPluginTemplates() const { return AddedPluginTemplates; }
 
+	const TArray<TPair<FOnPluginBeingEdited, FPluginEditorExtensionHandle>>& GetCustomizePluginEditingDelegates() { return CustomizePluginEditingDelegates; }
+
 private:
 	/** Called to spawn the plugin browser tab */
 	TSharedRef<SDockTab> HandleSpawnPluginBrowserTab(const FSpawnTabArgs& SpawnTabArgs);
@@ -95,6 +100,10 @@ private:
 	/** List of added plugin templates */
 	TArray<TSharedRef<FPluginTemplateDescription>> AddedPluginTemplates;
 
+	/** Additional customizers of plugin editing */
+	TArray<TPair<FOnPluginBeingEdited, FPluginEditorExtensionHandle>> CustomizePluginEditingDelegates;
+	int32 EditorExtensionCounter = 0;
+
 	/** The spawned browser tab */
 	TWeakPtr<SDockTab> PluginBrowserTab;
 
diff --git a/Engine/Plugins/Editor/PluginBrowser/Source/PluginBrowser/Private/PluginDescriptorEditor.cpp b/Engine/Plugins/Editor/PluginBrowser/Source/PluginBrowser/Private/PluginDescriptorEditor.cpp
new file mode 100644
--- /dev/null
+++ b/Engine/Plugins/Editor/PluginBrowser/Source/PluginBrowser/Private/PluginDescriptorEditor.cpp
@@ -0,0 +1,149 @@
+// Copyright Epic Games, Inc. All Rights Reserved.
+
+#include "PluginDescriptorEditor.h"
+#include "HAL/PlatformFileManager.h"
+#include "Misc/MessageDialog.h"
+#include "Misc/App.h"
+#include "Modules/ModuleManager.h"
+#include "Widgets/SWindow.h"
+#include "Framework/Application/SlateApplication.h"
+// #include "Widgets/Images/SImage.h"
+// #include "Widgets/Text/STextBlock.h"
+// #include "Widgets/Input/SButton.h"
+// #include "Widgets/Input/SCheckBox.h"
+#include "EditorStyleSet.h"
+#include "ISourceControlOperation.h"
+#include "SourceControlOperations.h"
+#include "ISourceControlProvider.h"
+#include "ISourceControlModule.h"
+#include "PluginDescriptor.h"
+#include "Interfaces/IPluginManager.h"
+// #include "SPluginBrowser.h"
+#include "PluginStyle.h"
+// #include "GameProjectGenerationModule.h"
+// #include "IDetailsView.h"
+// #include "Widgets/Input/SHyperlink.h"
+#include "PluginMetadataObject.h"
+#include "Interfaces/IProjectManager.h"
+// #include "PluginBrowserModule.h"
+#include "PropertyEditorModule.h"
+// #include "IUATHelperModule.h"
+// #include "DesktopPlatformModule.h"
+// #include "Widgets/Layout/SSpacer.h"
+
+#define LOCTEXT_NAMESPACE "PluginListTile"
+
+void FPluginDescriptorEditor::OpenEditorWindow(TSharedRef<IPlugin> Plugin, TSharedPtr<SWidget> ParentWidget, FSimpleDelegate OnEditCommitted)
+{
+	// Construct the plugin metadata object using the descriptor for this plugin
+	UPluginMetadataObject* MetadataObject = NewObject<UPluginMetadataObject>();
+	MetadataObject->TargetIconPath = Plugin->GetBaseDir() / TEXT("Resources/Icon128.png");
+	MetadataObject->PopulateFromPlugin(Plugin);
+	MetadataObject->AddToRoot();
+
+	// Create a property view
+	FPropertyEditorModule& EditModule = FModuleManager::Get().GetModuleChecked<FPropertyEditorModule>("PropertyEditor");
+	
+	FDetailsViewArgs DetailsViewArgs;
+	DetailsViewArgs.bAllowSearch = false;
+	DetailsViewArgs.bHideSelectionTip = true;
+
+	TSharedRef<IDetailsView> PropertyView = EditModule.CreateDetailView(DetailsViewArgs);
+	PropertyView->SetObject(MetadataObject, true);
+
+	// Create the window
+	TSharedRef<SWindow> PropertiesWindow = SNew(SWindow)
+		.SupportsMaximize(false)
+		.SupportsMinimize(false)
+		.SizingRule(ESizingRule::UserSized)
+		.ClientSize(FVector2D(700.0f, 700.0f))
+		.Title(LOCTEXT("PluginMetadata", "Plugin Properties"));
+
+	PropertiesWindow->SetContent(
+			SNew(SBorder)
+			.Padding(FMargin(8.0f, 8.0f))
+			[
+				SNew(SVerticalBox)
+
+				+SVerticalBox::Slot()
+				.AutoHeight()
+				.Padding(FMargin(5.0f, 10.0f, 5.0f, 5.0f))
+				[
+					SNew(STextBlock)
+					.Font(FPluginStyle::Get()->GetFontStyle(TEXT("PluginMetadataNameFont")))
+					.Text(FText::FromString(Plugin->GetName()))
+				]
+
+				+ SVerticalBox::Slot()
+				.Padding(5)
+				[
+					PropertyView
+				]
+
+				+ SVerticalBox::Slot()
+				.AutoHeight()
+				.Padding(5)
+				.HAlign(HAlign_Right)
+				[
+					SNew(SButton)
+					.ContentPadding(FMargin(20.0f, 2.0f))
+					.Text(LOCTEXT("OkButtonLabel", "Ok"))
+					.OnClicked(FOnClicked::CreateStatic(&FPluginDescriptorEditor::OnEditPluginFinished, MetadataObject, TSharedPtr<IPlugin>(Plugin), OnEditCommitted, TWeakPtr<SWindow>(PropertiesWindow)))
+				]
+			]
+		);
+
+	FSlateApplication::Get().AddModalWindow(PropertiesWindow, ParentWidget);
+}
+
+FReply FPluginDescriptorEditor::OnEditPluginFinished(UPluginMetadataObject* MetadataObject, TSharedPtr<IPlugin> Plugin, FSimpleDelegate OnEditCommitted, TWeakPtr<SWindow> WeakWindow)
+{
+	FPluginDescriptor OldDescriptor = Plugin->GetDescriptor();
+
+	// Update the descriptor with the new metadata
+	FPluginDescriptor NewDescriptor = OldDescriptor;
+	MetadataObject->CopyIntoDescriptor(NewDescriptor);
+	MetadataObject->RemoveFromRoot();
+
+	// Close the properties window
+	TSharedPtr<SWindow> PropertiesWindow = WeakWindow.Pin();
+	if (PropertiesWindow.IsValid())
+	{
+		PropertiesWindow->RequestDestroyWindow();
+	}
+
+	// Write both to strings
+	FString OldText;
+	OldDescriptor.Write(OldText);
+	FString NewText;
+	NewDescriptor.Write(NewText);
+	if(OldText.Compare(NewText, ESearchCase::CaseSensitive) != 0)
+	{
+		FString DescriptorFileName = Plugin->GetDescriptorFileName();
+
+		// First attempt to check out the file if SCC is enabled
+		ISourceControlModule& SourceControlModule = ISourceControlModule::Get();
+		if(SourceControlModule.IsEnabled())
+		{
+			ISourceControlProvider& SourceControlProvider = SourceControlModule.GetProvider();
+			TSharedPtr<ISourceControlState, ESPMode::ThreadSafe> SourceControlState = SourceControlProvider.GetState(DescriptorFileName, EStateCacheUsage::ForceUpdate);
+			if(SourceControlState.IsValid() && SourceControlState->CanCheckout())
+			{
+				SourceControlProvider.Execute(ISourceControlOperation::Create<FCheckOut>(), DescriptorFileName);
+			}
+		}
+
+		// Write to the file and update the in-memory metadata
+		FText FailReason;
+		if(!Plugin->UpdateDescriptor(NewDescriptor, FailReason))
+		{
+			FMessageDialog::Open(EAppMsgType::Ok, FailReason);
+		}
+	}
+
+	OnEditCommitted.ExecuteIfBound();
+
+	return FReply::Handled();
+}
+
+#undef LOCTEXT_NAMESPACE
diff --git a/Engine/Plugins/Editor/PluginBrowser/Source/PluginBrowser/Private/PluginDescriptorEditor.h b/Engine/Plugins/Editor/PluginBrowser/Source/PluginBrowser/Private/PluginDescriptorEditor.h
new file mode 100644
--- /dev/null
+++ b/Engine/Plugins/Editor/PluginBrowser/Source/PluginBrowser/Private/PluginDescriptorEditor.h
@@ -0,0 +1,23 @@
+// Copyright Epic Games, Inc. All Rights Reserved.
+
+#pragma once
+
+#include "CoreMinimal.h"
+
+class UPluginMetadataObject;
+class IPlugin;
+class SWindow;
+class SWidget;
+
+/**
+ * Creates a window to edit a .uplugin file
+ */
+class FPluginDescriptorEditor
+{
+public:
+	static void OpenEditorWindow(TSharedRef<IPlugin> PluginToEdit, TSharedPtr<SWidget> ParentWidget, FSimpleDelegate OnEditCommitted);
+
+private:
+	static FReply OnEditPluginFinished(UPluginMetadataObject* MetadataObject, TSharedPtr<IPlugin> Plugin, FSimpleDelegate OnEditCommitted, TWeakPtr<SWindow> WeakWindow);
+};
+
diff --git a/Engine/Plugins/Editor/PluginBrowser/Source/PluginBrowser/Private/PluginMetadataObject.cpp b/Engine/Plugins/Editor/PluginBrowser/Source/PluginBrowser/Private/PluginMetadataObject.cpp
--- a/Engine/Plugins/Editor/PluginBrowser/Source/PluginBrowser/Private/PluginMetadataObject.cpp
+++ b/Engine/Plugins/Editor/PluginBrowser/Source/PluginBrowser/Private/PluginMetadataObject.cpp
@@ -12,6 +12,9 @@
 #include "DetailWidgetRow.h"
 #include "DetailCategoryBuilder.h"
 #include "SExternalImageReference.h"
+#include "Dom/JsonObject.h"
+#include "Features/IPluginsEditorFeature.h"
+#include "PluginBrowserModule.h"
 
 void FPluginReferenceMetadata::PopulateFromDescriptor(const FPluginReferenceDescriptor& InDescriptor)
 {
@@ -31,8 +34,11 @@ UPluginMetadataObject::UPluginMetadataObject(const FObjectInitializer& ObjectIni
 {
 }
 
-void UPluginMetadataObject::PopulateFromDescriptor(const FPluginDescriptor& InDescriptor)
+void UPluginMetadataObject::PopulateFromPlugin(TSharedPtr<IPlugin> InPlugin)
 {
+	SourcePlugin = InPlugin;
+
+	const FPluginDescriptor& InDescriptor = InPlugin->GetDescriptor();
 	Version = InDescriptor.Version;
 	VersionName = InDescriptor.VersionName;
 	FriendlyName = InDescriptor.FriendlyName;
@@ -91,6 +97,12 @@ void UPluginMetadataObject::CopyIntoDescriptor(FPluginDescriptor& OutDescriptor)
 	}
 
 	OutDescriptor.Plugins = MoveTemp(NewPlugins);
+	
+	// Apply any edits done by an extension
+	for (const TSharedPtr<FPluginEditorExtension>& Extension : Extensions)
+	{
+		Extension->CommitEdits(OutDescriptor);
+	}
 }
 
 /////////////////////////////////////////////////////////////////////////////////////////////////////////
@@ -108,7 +120,22 @@ void FPluginMetadataCustomization::CustomizeDetails(IDetailLayoutBuilder& Detail
 	if(Objects.Num() == 1 && Objects[0].IsValid())
 	{
 		UPluginMetadataObject* PluginMetadata = Cast<UPluginMetadataObject>(Objects[0].Get());
-		if(PluginMetadata != nullptr && PluginMetadata->TargetIconPath.Len() > 0)
+		check(PluginMetadata);
+
+		// Run any external customizations
+		PluginMetadata->Extensions.Reset();
+		FPluginEditingContext PluginEditingContext;
+		PluginEditingContext.PluginBeingEdited = PluginMetadata->SourcePlugin.Pin();
+		for (const auto& KVP : FPluginBrowserModule::Get().GetCustomizePluginEditingDelegates())
+		{
+			TSharedPtr<FPluginEditorExtension> Extension = KVP.Key.Execute(PluginEditingContext, DetailBuilder);
+			if (Extension.IsValid())
+			{
+				PluginMetadata->Extensions.Add(Extension);
+			}
+		}
+		
+		if(PluginMetadata->TargetIconPath.Len() > 0)
 		{
 			// Get the current icon path
 			FString CurrentIconPath = PluginMetadata->TargetIconPath;
diff --git a/Engine/Plugins/Editor/PluginBrowser/Source/PluginBrowser/Private/PluginMetadataObject.h b/Engine/Plugins/Editor/PluginBrowser/Source/PluginBrowser/Private/PluginMetadataObject.h
--- a/Engine/Plugins/Editor/PluginBrowser/Source/PluginBrowser/Private/PluginMetadataObject.h
+++ b/Engine/Plugins/Editor/PluginBrowser/Source/PluginBrowser/Private/PluginMetadataObject.h
@@ -11,6 +11,8 @@
 class IDetailLayoutBuilder;
 struct FPluginDescriptor;
 struct FPluginReferenceDescriptor;
+class IPlugin;
+struct FPluginEditorExtension;
 
 /**
  * We use this object to display plugin reference properties using details view.
@@ -114,10 +116,16 @@ public:
 	UPROPERTY(EditAnywhere, Category = Dependencies)
 	TArray<FPluginReferenceMetadata> Plugins;
 
+	/** Plugin this proxy object was constructed from */
+	TWeakPtr<IPlugin> SourcePlugin;
+
+	/** Editing extensions */
+	TArray<TSharedPtr<FPluginEditorExtension>> Extensions;
+
 	/**
 	 * Populate the fields of this object from an existing descriptor.
 	 */
-	void PopulateFromDescriptor(const FPluginDescriptor& InDescriptor);
+	void PopulateFromPlugin(TSharedPtr<IPlugin> InPlugin);
 
 	/**
 	 * Copy the metadata fields into a plugin descriptor.
diff --git a/Engine/Plugins/Editor/PluginBrowser/Source/PluginBrowser/Private/SNewPluginWizard.cpp b/Engine/Plugins/Editor/PluginBrowser/Source/PluginBrowser/Private/SNewPluginWizard.cpp
--- a/Engine/Plugins/Editor/PluginBrowser/Source/PluginBrowser/Private/SNewPluginWizard.cpp
+++ b/Engine/Plugins/Editor/PluginBrowser/Source/PluginBrowser/Private/SNewPluginWizard.cpp
@@ -640,25 +640,35 @@ FReply SNewPluginWizard::OnCreatePluginClicked()
 	const FString PluginName = PluginNameText.ToString();
 	const bool bHasModules = PluginWizardDefinition->HasModules();
 	
-	FPluginUtils::FNewPluginParams CreationParams;
+	FPluginUtils::FNewPluginParamsWithDescriptor CreationParams;
 	CreationParams.TemplateFolders = PluginWizardDefinition->GetFoldersForSelection();
-	CreationParams.bCanContainContent = Template->bCanContainContent;
-	CreationParams.bHasModules = bHasModules;
-	CreationParams.ModuleDescriptorType = PluginWizardDefinition->GetPluginModuleDescriptor();
-	CreationParams.LoadingPhase = PluginWizardDefinition->GetPluginLoadingPhase();
-	PluginWizardDefinition->GetPluginIconPath(CreationParams.PluginIconPath);
+	CreationParams.Descriptor.bCanContainContent = Template->bCanContainContent;
+
+	if (bHasModules)
+	{
+		CreationParams.Descriptor.Modules.Add(FModuleDescriptor(*PluginName, PluginWizardDefinition->GetPluginModuleDescriptor(), PluginWizardDefinition->GetPluginLoadingPhase()));
+	}
+
+	CreationParams.Descriptor.FriendlyName = PluginName;
+	CreationParams.Descriptor.Version = 1;
+	CreationParams.Descriptor.VersionName = TEXT("1.0");
+	CreationParams.Descriptor.Category = TEXT("Other");
+
+	PluginWizardDefinition->GetPluginIconPath(/*out*/ CreationParams.PluginIconPath);
 	if (DescriptorData.IsValid())
 	{
-		CreationParams.CreatedBy = DescriptorData->CreatedBy;
-		CreationParams.CreatedByURL = DescriptorData->CreatedByURL;
-		CreationParams.Description = DescriptorData->Description;
-		CreationParams.bIsBetaVersion = DescriptorData->bIsBetaVersion;
+		CreationParams.Descriptor.CreatedBy = DescriptorData->CreatedBy;
+		CreationParams.Descriptor.CreatedByURL = DescriptorData->CreatedByURL;
+		CreationParams.Descriptor.Description = DescriptorData->Description;
+		CreationParams.Descriptor.bIsBetaVersion = DescriptorData->bIsBetaVersion;
 	}
 
 	FPluginUtils::FMountPluginParams MountParams;
 	MountParams.bEnablePluginInProject = true;
 	MountParams.bUpdateProjectPluginSearchPath = true;
 	MountParams.bSelectInContentBrowser = ShowPluginContentDirectoryCheckBox->IsChecked();
+
+	Template->CustomizeDescriptorBeforeCreation(CreationParams.Descriptor);
 	
 	FText FailReason;
 	TSharedPtr<IPlugin> NewPlugin = FPluginUtils::CreateAndMountNewPlugin(PluginName, PluginFolderPath, CreationParams, MountParams, FailReason);
@@ -669,8 +679,6 @@ FReply SNewPluginWizard::OnCreatePluginClicked()
 
 	if (bSucceeded)
 	{
-		//@TODO: Add the plugin files to source control if the project is configured for it!
-
 		// Let the template create additional assets / modify state after creation
 		Template->OnPluginCreated(NewPlugin);
 
diff --git a/Engine/Plugins/Editor/PluginBrowser/Source/PluginBrowser/Private/SPluginTile.cpp b/Engine/Plugins/Editor/PluginBrowser/Source/PluginBrowser/Private/SPluginTile.cpp
--- a/Engine/Plugins/Editor/PluginBrowser/Source/PluginBrowser/Private/SPluginTile.cpp
+++ b/Engine/Plugins/Editor/PluginBrowser/Source/PluginBrowser/Private/SPluginTile.cpp
@@ -582,117 +582,19 @@ EVisibility SPluginTile::GetAuthoringButtonsVisibility() const
 
 void SPluginTile::OnEditPlugin()
 {
-	// Construct the plugin metadata object using the descriptor for this plugin
-	UPluginMetadataObject* MetadataObject = NewObject<UPluginMetadataObject>();
-	MetadataObject->TargetIconPath = Plugin->GetBaseDir() / TEXT("Resources/Icon128.png");
-	MetadataObject->PopulateFromDescriptor(Plugin->GetDescriptor());
-	MetadataObject->AddToRoot();
-
-	// Create a property view
-	FPropertyEditorModule& EditModule = FModuleManager::Get().GetModuleChecked<FPropertyEditorModule>("PropertyEditor");
-	
-	FDetailsViewArgs DetailsViewArgs;
-	DetailsViewArgs.bAllowSearch = false;
-	DetailsViewArgs.bHideSelectionTip = true;
-
-	TSharedRef<IDetailsView> PropertyView = EditModule.CreateDetailView(DetailsViewArgs);
-	PropertyView->SetObject(MetadataObject, true);
-
-	// Create the window
-	PropertiesWindow = SNew(SWindow)
-		.SupportsMaximize(false)
-		.SupportsMinimize(false)
-		.SizingRule(ESizingRule::UserSized)
-		.ClientSize(FVector2D(700.0f, 700.0f))
-		.Title(LOCTEXT("PluginMetadata", "Plugin Properties"))
-		.Content()
-		[
-			SNew(SBorder)
-			.Padding(FMargin(8.0f, 8.0f))
-			[
-				SNew(SVerticalBox)
-
-				+SVerticalBox::Slot()
-				.AutoHeight()
-				.Padding(FMargin(5.0f, 10.0f, 5.0f, 5.0f))
-				[
-					SNew(STextBlock)
-					.Font(FPluginStyle::Get()->GetFontStyle(TEXT("PluginMetadataNameFont")))
-					.Text(FText::FromString(Plugin->GetName()))
-				]
-
-				+ SVerticalBox::Slot()
-				.Padding(5)
-				[
-					PropertyView
-				]
-
-				+ SVerticalBox::Slot()
-				.AutoHeight()
-				.Padding(5)
-				.HAlign(HAlign_Right)
-				[
-					SNew(SButton)
-					.ContentPadding(FMargin(20.0f, 2.0f))
-					.Text(LOCTEXT("OkButtonLabel", "Ok"))
-					.OnClicked(this, &SPluginTile::OnEditPluginFinished, MetadataObject)
-				]
-			]
-		];
-
-	FSlateApplication::Get().AddModalWindow(PropertiesWindow.ToSharedRef(), OwnerWeak.Pin());//Args.ParentWidget);
+	FPluginBrowserModule::Get().OpenPluginEditor(Plugin.ToSharedRef(), OwnerWeak.Pin(), FSimpleDelegate::CreateRaw(this, &SPluginTile::OnEditPluginFinished));
 }
 
-FReply SPluginTile::OnEditPluginFinished(UPluginMetadataObject* MetadataObject)
+void SPluginTile::OnEditPluginFinished()
 {
-	FPluginDescriptor OldDescriptor = Plugin->GetDescriptor();
-
-	// Update the descriptor with the new metadata
-	FPluginDescriptor NewDescriptor = OldDescriptor;
-	MetadataObject->CopyIntoDescriptor(NewDescriptor);
-	MetadataObject->RemoveFromRoot();
-
-	// Close the properties window
-	PropertiesWindow->RequestDestroyWindow();
-
-	// Write both to strings
-	FString OldText;
-	OldDescriptor.Write(OldText);
-	FString NewText;
-	NewDescriptor.Write(NewText);
-	if(OldText.Compare(NewText, ESearchCase::CaseSensitive) != 0)
-	{
-		FString DescriptorFileName = Plugin->GetDescriptorFileName();
-
-		// First attempt to check out the file if SCC is enabled
-		ISourceControlModule& SourceControlModule = ISourceControlModule::Get();
-		if(SourceControlModule.IsEnabled())
-		{
-			ISourceControlProvider& SourceControlProvider = SourceControlModule.GetProvider();
-			TSharedPtr<ISourceControlState, ESPMode::ThreadSafe> SourceControlState = SourceControlProvider.GetState(DescriptorFileName, EStateCacheUsage::ForceUpdate);
-			if(SourceControlState.IsValid() && SourceControlState->CanCheckout())
-			{
-				SourceControlProvider.Execute(ISourceControlOperation::Create<FCheckOut>(), DescriptorFileName);
-			}
-		}
-
-		// Write to the file and update the in-memory metadata
-		FText FailReason;
-		if(!Plugin->UpdateDescriptor(NewDescriptor, FailReason))
-		{
-			FMessageDialog::Open(EAppMsgType::Ok, FailReason);
-		}
-
-		// Recreate the widgets on this tile
-		RecreateWidgets();
+	// Recreate the widgets on this tile
+	RecreateWidgets();
 
-		// Refresh the parent too
-		if(OwnerWeak.IsValid())
-		{
-			OwnerWeak.Pin()->GetOwner().SetNeedsRefresh();
-		}
+	// Refresh the parent too
+	if(OwnerWeak.IsValid())
+	{
+		OwnerWeak.Pin()->GetOwner().SetNeedsRefresh();
 	}
-	return FReply::Handled();
 }
 
 void SPluginTile::OnPackagePlugin()
diff --git a/Engine/Plugins/Editor/PluginBrowser/Source/PluginBrowser/Private/SPluginTile.h b/Engine/Plugins/Editor/PluginBrowser/Source/PluginBrowser/Private/SPluginTile.h
--- a/Engine/Plugins/Editor/PluginBrowser/Source/PluginBrowser/Private/SPluginTile.h
+++ b/Engine/Plugins/Editor/PluginBrowser/Source/PluginBrowser/Private/SPluginTile.h
@@ -11,7 +11,6 @@
 #include "Brushes/SlateDynamicImageBrush.h"
 
 class IPlugin;
-class SWindow;
 class UPluginMetadataObject;
 
 /**
@@ -49,8 +48,8 @@ private:
 	/** Called when the 'edit' hyperlink is clicked */
 	void OnEditPlugin();
 
-	/** Called when the 'edit' hyperlink is clicked */
-	FReply OnEditPluginFinished(UPluginMetadataObject* MetadataObject);
+	/** Called when the edit window is committed */
+	void OnEditPluginFinished();
 
 	/** Called when the 'package' hyperlink is clicked */
 	void OnPackagePlugin();
@@ -63,9 +62,6 @@ private:
 	/** Weak pointer back to its owner */
 	TWeakPtr< class SPluginTileList > OwnerWeak;
 
-	/** Dialog for editing plugin properties */
-	TSharedPtr< SWindow > PropertiesWindow;
-
 	/** Brush resource for the image that is dynamically loaded */
 	TSharedPtr< FSlateDynamicImageBrush > PluginIconDynamicImageBrush;
 };
diff --git a/Engine/Plugins/Experimental/GameFeatures/Source/GameFeatures/Private/GameFeaturesSubsystem.cpp b/Engine/Plugins/Experimental/GameFeatures/Source/GameFeatures/Private/GameFeaturesSubsystem.cpp
--- a/Engine/Plugins/Experimental/GameFeatures/Source/GameFeatures/Private/GameFeaturesSubsystem.cpp
+++ b/Engine/Plugins/Experimental/GameFeatures/Source/GameFeatures/Private/GameFeaturesSubsystem.cpp
@@ -487,27 +487,9 @@ void UGameFeaturesSubsystem::LoadBuiltInGameFeaturePlugin(const TSharedRef<IPlug
 				{
 					UGameFeaturePluginStateMachine* StateMachine = GetGameFeaturePluginStateMachine(PluginURL, true);
 
-					EBuiltInAutoState InitialAutoState = BehaviorOptions.AutoStateOverride != EBuiltInAutoState::Invalid ? BehaviorOptions.AutoStateOverride : PluginDetails.BuiltInAutoState;
+					const EBuiltInAutoState InitialAutoState = (BehaviorOptions.AutoStateOverride != EBuiltInAutoState::Invalid) ? BehaviorOptions.AutoStateOverride : PluginDetails.BuiltInAutoState;
 						
-					EGameFeaturePluginState DestinationState = EGameFeaturePluginState::UnknownStatus;
-					switch (InitialAutoState)
-					{
-					case EBuiltInAutoState::Installed:
-						DestinationState = EGameFeaturePluginState::Installed;
-						break;
-					case EBuiltInAutoState::Registered:
-						break;
-					case EBuiltInAutoState::Loaded:
-						DestinationState = EGameFeaturePluginState::Loaded;
-						break;
-					case EBuiltInAutoState::Active:
-						DestinationState = EGameFeaturePluginState::Active;
-						break;
-
-					default:
-						ensureMsgf(false, TEXT("BuilIn Game Feature Plugin %s configured for invalid auto state %d. Setting to UnknownStatus."), (uint8)InitialAutoState);
-						break;
-					}
+					const EGameFeaturePluginState DestinationState = ConvertInitialFeatureStateToTargetState(InitialAutoState);
 
 					if (StateMachine->GetCurrentState() >= DestinationState)
 					{
@@ -626,36 +608,17 @@ bool UGameFeaturesSubsystem::GetGameFeaturePluginDetails(const FString& PluginDe
 		return false;
 	}
 
+	//@TODO: When we properly support downloaded plugins, will need to determine this
+	const bool bIsBuiltInPlugin = true;
+
 	// Read the properties
 
 	// Hotfixable. If it is not specified, then we assume it is
 	OutPluginDetails.bHotfixable = true;
 	ObjectPtr->TryGetBoolField(TEXT("Hotfixable"), OutPluginDetails.bHotfixable);
 
-	// BuiltInAutoRegister. Default to true. If this is a built in plugin, should it be registered automatically (set to false if you intent to load late with LoadAndActivateGameFeaturePlugin)
-	bool bBuiltInAutoRegister = true;
-	ObjectPtr->TryGetBoolField(TEXT("BuiltInAutoRegister"), bBuiltInAutoRegister);
-
-	// BuiltInAutoLoad. Default to true. If this is a built in plugin, should it be loaded automatically (set to false if you intent to load late with LoadAndActivateGameFeaturePlugin)
-	bool bBuiltInAutoLoad = true;
-	ObjectPtr->TryGetBoolField(TEXT("BuiltInAutoLoad"), bBuiltInAutoLoad);
-
-	// The cooker will need to activate the plugin so that assets can be scanned properly
-	bool bBuiltInAutoActivate = true;
-	ObjectPtr->TryGetBoolField(TEXT("BuiltInAutoActivate"), bBuiltInAutoActivate);
-
-	if (bBuiltInAutoRegister)
-	{
-		OutPluginDetails.BuiltInAutoState = EBuiltInAutoState::Registered;
-		if (bBuiltInAutoLoad)
-		{
-			OutPluginDetails.BuiltInAutoState = EBuiltInAutoState::Loaded;
-			if (bBuiltInAutoActivate)
-			{
-				OutPluginDetails.BuiltInAutoState = EBuiltInAutoState::Active;
-			}
-		}
-	}
+	// Determine the initial plugin state
+	OutPluginDetails.BuiltInAutoState = bIsBuiltInPlugin ? DetermineBuiltInInitialFeatureState(ObjectPtr, PluginDescriptorFilename) : EBuiltInAutoState::Installed;
 
 	// Read any additional metadata the policy might want to consume (e.g., a release version number)
 	for (const FString& ExtraKey : GetDefault<UGameFeaturesSubsystemSettings>()->AdditionalPluginMetadataKeys)
@@ -665,6 +628,7 @@ bool UGameFeaturesSubsystem::GetGameFeaturePluginDetails(const FString& PluginDe
 		OutPluginDetails.AdditionalMetadata.Add(ExtraKey, ExtraValue);
 	}
 
+	// Parse plugin dependencies
 	const TArray<TSharedPtr<FJsonValue>>* PluginsArray = nullptr;
 	ObjectPtr->TryGetArrayField(TEXT("Plugins"), PluginsArray);
 	if (PluginsArray)
@@ -853,3 +817,102 @@ void UGameFeaturesSubsystem::ListGameFeaturePlugins(const TArray<FString>& Args,
 
 	Ar.Logf(TEXT("Total Game Feature Plugins: %d"), PluginCount);
 }
+
+EBuiltInAutoState UGameFeaturesSubsystem::DetermineBuiltInInitialFeatureState(TSharedPtr<FJsonObject> Descriptor, const FString& ErrorContext)
+{
+	EBuiltInAutoState InitialState = EBuiltInAutoState::Invalid;
+
+	FString InitialFeatureStateStr;
+	if (Descriptor->TryGetStringField(TEXT("BuiltInInitialFeatureState"), InitialFeatureStateStr))
+	{
+		if (InitialFeatureStateStr == TEXT("Installed"))
+		{
+			InitialState = EBuiltInAutoState::Installed;
+		}
+		else if (InitialFeatureStateStr == TEXT("Registered"))
+		{
+			InitialState = EBuiltInAutoState::Registered;
+		}
+		else if (InitialFeatureStateStr == TEXT("Loaded"))
+		{
+			InitialState = EBuiltInAutoState::Loaded;
+		}
+		else if (InitialFeatureStateStr == TEXT("Active"))
+		{
+			InitialState = EBuiltInAutoState::Active;
+		}
+		else
+		{
+			if (!ErrorContext.IsEmpty())
+			{
+				UE_LOG(LogGameFeatures, Error, TEXT("Game feature '%s' has an unknown value '%s' for BuiltInInitialFeatureState (expected Installed, Registered, Loaded, or Active); defaulting to Active."), *ErrorContext, *InitialFeatureStateStr);
+			}
+			InitialState = EBuiltInAutoState::Active;
+		}
+	}
+	else
+	{
+		// BuiltInAutoRegister. Default to true. If this is a built in plugin, should it be registered automatically (set to false if you intent to load late with LoadAndActivateGameFeaturePlugin)
+		bool bBuiltInAutoRegister = true;
+		Descriptor->TryGetBoolField(TEXT("BuiltInAutoRegister"), bBuiltInAutoRegister);
+
+		// BuiltInAutoLoad. Default to true. If this is a built in plugin, should it be loaded automatically (set to false if you intent to load late with LoadAndActivateGameFeaturePlugin)
+		bool bBuiltInAutoLoad = true;
+		Descriptor->TryGetBoolField(TEXT("BuiltInAutoLoad"), bBuiltInAutoLoad);
+
+		// The cooker will need to activate the plugin so that assets can be scanned properly
+		bool bBuiltInAutoActivate = true;
+		Descriptor->TryGetBoolField(TEXT("BuiltInAutoActivate"), bBuiltInAutoActivate);
+
+		InitialState = EBuiltInAutoState::Installed;
+		if (bBuiltInAutoRegister)
+		{
+			InitialState = EBuiltInAutoState::Registered;
+			if (bBuiltInAutoLoad)
+			{
+				InitialState = EBuiltInAutoState::Loaded;
+				if (bBuiltInAutoActivate)
+				{
+					InitialState = EBuiltInAutoState::Active;
+				}
+			}
+		}
+
+		if (!ErrorContext.IsEmpty())
+		{
+			//@TODO: Increase severity to a warning after changing existing features
+			UE_LOG(LogGameFeatures, Log, TEXT("Game feature '%s' has no BuiltInInitialFeatureState key, using legacy BuiltInAutoRegister(%d)/BuiltInAutoLoad(%d)/BuiltInAutoActivate(%d) values to arrive at initial state."),
+				*ErrorContext,
+				bBuiltInAutoRegister ? 1 : 0,
+				bBuiltInAutoLoad ? 1 : 0,
+				bBuiltInAutoActivate ? 1 : 0);
+		}
+	}
+
+	return InitialState;
+}
+
+EGameFeaturePluginState UGameFeaturesSubsystem::ConvertInitialFeatureStateToTargetState(EBuiltInAutoState AutoState)
+{
+	EGameFeaturePluginState InitialState;
+	switch (AutoState)
+	{
+	default:
+	case EBuiltInAutoState::Invalid:
+		InitialState = EGameFeaturePluginState::UnknownStatus;
+		break;
+	case EBuiltInAutoState::Installed:
+		InitialState = EGameFeaturePluginState::Installed;
+		break;
+	case EBuiltInAutoState::Registered:
+		InitialState = EGameFeaturePluginState::Registered;
+		break;
+	case EBuiltInAutoState::Loaded:
+		InitialState = EGameFeaturePluginState::Loaded;
+		break;
+	case EBuiltInAutoState::Active:
+		InitialState = EGameFeaturePluginState::Active;
+		break;
+	}
+	return InitialState;
+}
\ No newline at end of file
diff --git a/Engine/Plugins/Experimental/GameFeatures/Source/GameFeatures/Public/GameFeaturesSubsystem.h b/Engine/Plugins/Experimental/GameFeatures/Source/GameFeatures/Public/GameFeaturesSubsystem.h
--- a/Engine/Plugins/Experimental/GameFeatures/Source/GameFeatures/Public/GameFeaturesSubsystem.h
+++ b/Engine/Plugins/Experimental/GameFeatures/Source/GameFeatures/Public/GameFeaturesSubsystem.h
@@ -16,6 +16,7 @@ class UGameFeatureData;
 class UGameFeaturesProjectPolicies;
 enum class EGameFeaturePluginState : uint8;
 class IPlugin;
+class FJsonObject;
 
 struct FGameFeatureDeactivatingContext
 {
@@ -193,6 +194,11 @@ public:
 	/** Returns the current state of the state machine for the specified plugin URL */
 	EGameFeaturePluginState GetPluginState(const FString& PluginURL);
 
+	/** Determine the initial feature state for a built-in plugin */
+	static EBuiltInAutoState DetermineBuiltInInitialFeatureState(TSharedPtr<FJsonObject> Descriptor, const FString& ErrorContext);
+
+	static EGameFeaturePluginState ConvertInitialFeatureStateToTargetState(EBuiltInAutoState InitialState);
+
 private:
 	void OnGameFeatureRegistering(const UGameFeatureData* GameFeatureData, const FString& PluginName);
 	friend struct FGameFeaturePluginState_Registering;
diff --git a/Engine/Plugins/Experimental/GameFeatures/Source/GameFeaturesEditor/GameFeaturesEditor.Build.cs b/Engine/Plugins/Experimental/GameFeatures/Source/GameFeaturesEditor/GameFeaturesEditor.Build.cs
--- a/Engine/Plugins/Experimental/GameFeatures/Source/GameFeaturesEditor/GameFeaturesEditor.Build.cs
+++ b/Engine/Plugins/Experimental/GameFeatures/Source/GameFeaturesEditor/GameFeaturesEditor.Build.cs
@@ -29,7 +29,8 @@ namespace UnrealBuildTool.Rules
 					"SlateCore",
 					"EditorStyle",
 					"PropertyEditor",
-					"SharedSettingsWidgets"
+					"SharedSettingsWidgets",
+					"Json"
 				}
 			);
 		}
diff --git a/Engine/Plugins/Experimental/GameFeatures/Source/GameFeaturesEditor/Private/GameFeatureDataDetailsCustomization.cpp b/Engine/Plugins/Experimental/GameFeatures/Source/GameFeaturesEditor/Private/GameFeatureDataDetailsCustomization.cpp
--- a/Engine/Plugins/Experimental/GameFeatures/Source/GameFeaturesEditor/Private/GameFeatureDataDetailsCustomization.cpp
+++ b/Engine/Plugins/Experimental/GameFeatures/Source/GameFeaturesEditor/Private/GameFeatureDataDetailsCustomization.cpp
@@ -4,11 +4,16 @@
 #include "Widgets/DeclarativeSyntaxSupport.h"
 #include "Widgets/SBoxPanel.h"
 #include "Widgets/Input/SButton.h"
-#include "Widgets/Input/SSegmentedControl.h"
+#include "SGameFeatureStateWidget.h"
 #include "Widgets/Notifications/SErrorText.h"
 #include "DetailLayoutBuilder.h"
 #include "DetailWidgetRow.h"
 #include "DetailCategoryBuilder.h"
+
+#include "Interfaces/IPluginManager.h"
+#include "Features/IPluginsEditorFeature.h"
+#include "Features/EditorFeatures.h"
+
 #include "GameFeatureData.h"
 #include "../../GameFeatures/Private/GameFeaturePluginStateMachine.h"
 
@@ -28,7 +33,7 @@ void FGameFeatureDataDetailsCustomization::CustomizeDetails(IDetailLayoutBuilder
 		.ToolTipText(LOCTEXT("ErrorTooltip", "The error raised while attempting to change the state of this feature"));
 
 	// Create a category so this is displayed early in the properties
-	IDetailCategoryBuilder& TopCategory = DetailBuilder.EditCategory("Plugin Settings", FText::GetEmpty(), ECategoryPriority::Important);
+	IDetailCategoryBuilder& TopCategory = DetailBuilder.EditCategory("Feature State", FText::GetEmpty(), ECategoryPriority::Important);
 
 	PluginURL.Reset();
 	ObjectsBeingCustomized.Empty();
@@ -43,28 +48,65 @@ void FGameFeatureDataDetailsCustomization::CustomizeDetails(IDetailLayoutBuilder
 
 		UGameFeaturesSubsystem& Subsystem = UGameFeaturesSubsystem::Get();
 		Subsystem.GetPluginURLForBuiltInPluginByName(PathParts[0], /*out*/ PluginURL);
-
-		TSharedRef<SWidget> CurrentStateSwitcher = SNew(SSegmentedControl<EGameFeaturePluginState>)
-			.Value(this, &FGameFeatureDataDetailsCustomization::GetCurrentState)
-			.OnValueChanged(this, &FGameFeatureDataDetailsCustomization::ChangeDesiredState)
-			.ToolTipText(LOCTEXT("StateSwitcherTooltip", "Attempt to change the current state of this game feature"))
-			+SSegmentedControl<EGameFeaturePluginState>::Slot(EGameFeaturePluginState::Installed)
-				.Text(GetDisplayNameOfState(EGameFeaturePluginState::Installed))
-			+ SSegmentedControl<EGameFeaturePluginState>::Slot(EGameFeaturePluginState::Registered)
-				.Text(GetDisplayNameOfState(EGameFeaturePluginState::Registered))
-			+ SSegmentedControl<EGameFeaturePluginState>::Slot(EGameFeaturePluginState::Loaded)
-				.Text(GetDisplayNameOfState(EGameFeaturePluginState::Loaded))
-			+ SSegmentedControl<EGameFeaturePluginState>::Slot(EGameFeaturePluginState::Active)
-				.Text(GetDisplayNameOfState(EGameFeaturePluginState::Active))
-			;
+		PluginPtr = IPluginManager::Get().FindPlugin(PathParts[0]);
 
 		const float Padding = 8.0f;
 
+		if (PluginPtr.IsValid())
+		{
+			const FString ShortFilename = FPaths::GetCleanFilename(PluginPtr->GetDescriptorFileName());
+			FDetailWidgetRow& EditPluginRow = TopCategory.AddCustomRow(LOCTEXT("InitialStateSearchText", "Initial State Edit Plugin"))
+				.NameContent()
+				[
+					SNew(STextBlock)
+					.Text(LOCTEXT("InitialState", "Initial State"))
+					.ToolTipText(LOCTEXT("InitialStateTooltip", "The initial or default state of this game feature (determines the state that it will be in at game/editor startup)"))
+					.Font(DetailBuilder.GetDetailFont())
+				]
+
+				.ValueContent()
+				[
+					SNew(SHorizontalBox)
+					+SHorizontalBox::Slot()
+					.AutoWidth()
+					.Padding(0.0f, 0.0f, Padding, 0.0f)
+					.VAlign(VAlign_Center)
+					[
+						SNew(STextBlock)
+						.Text(this, &FGameFeatureDataDetailsCustomization::GetInitialStateText)
+						.Font(DetailBuilder.GetDetailFont())
+					]
+
+					+SHorizontalBox::Slot()
+					.AutoWidth()
+					.VAlign(VAlign_Center)
+					[
+						SNew(SButton)
+						.Text(LOCTEXT("EditPluginButton", "Edit Plugin"))
+						.OnClicked_Lambda([this]()
+							{
+								IModularFeatures& ModularFeatures = IModularFeatures::Get();
+								if (ModularFeatures.IsModularFeatureAvailable(EditorFeatures::PluginsEditor))
+								{
+									ModularFeatures.GetModularFeature<IPluginsEditorFeature>(EditorFeatures::PluginsEditor).OpenPluginEditor(PluginPtr.ToSharedRef(), nullptr, FSimpleDelegate());
+								}
+								else
+								{
+									FMessageDialog::Open(EAppMsgType::Ok, LOCTEXT("CannotEditPlugin_PluginBrowserDisabled", "Cannot open plugin editor because the PluginBrowser plugin is disabled)"));
+								}
+								return FReply::Handled();
+							})
+					]
+
+				];
+		}
+
 		FDetailWidgetRow& ControlRow = TopCategory.AddCustomRow(LOCTEXT("ControlSearchText", "Plugin State Control"))
 			.NameContent()
 			[
 				SNew(STextBlock)
 				.Text(LOCTEXT("CurrentState", "Current State"))
+				.ToolTipText(LOCTEXT("CurrentStateTooltip", "The current state of this game feature"))
 				.Font(DetailBuilder.GetDetailFont())
 			]
 
@@ -75,23 +117,10 @@ void FGameFeatureDataDetailsCustomization::CustomizeDetails(IDetailLayoutBuilder
 				+SVerticalBox::Slot()
 				.AutoHeight()
 				[
-					SNew(SHorizontalBox)
-					+SHorizontalBox::Slot()
-					.AutoWidth()
-					.VAlign(VAlign_Center)
-					[
-						CurrentStateSwitcher
-					]
-					+SHorizontalBox::Slot()
-					.Padding(8.0f, 0.0f, 0.0f, 0.0f)
-					.VAlign(VAlign_Center)
-					[
-						SNew(STextBlock)
-						.Text(this, &FGameFeatureDataDetailsCustomization::GetStateStatusDisplay)
-						.TextStyle(&FAppStyle::Get().GetWidgetStyle<FTextBlockStyle>("ButtonText"))
-						.ToolTipText(LOCTEXT("OtherStateToolTip", "The current state of this game feature plugin"))
-						.ColorAndOpacity(FAppStyle::Get().GetSlateColor(TEXT("Colors.AccentYellow")))
-					]
+					SNew(SGameFeatureStateWidget)
+					.ToolTipText(LOCTEXT("StateSwitcherTooltip", "Attempt to change the current state of this game feature"))
+					.CurrentState(this, &FGameFeatureDataDetailsCustomization::GetCurrentState)
+					.OnStateChanged(this, &FGameFeatureDataDetailsCustomization::ChangeDesiredState)
 				]
 				+SVerticalBox::Slot()
 				.HAlign(HAlign_Left)
@@ -178,54 +207,17 @@ void FGameFeatureDataDetailsCustomization::ChangeDesiredState(EGameFeaturePlugin
 	}
 }
 
-FText FGameFeatureDataDetailsCustomization::GetDisplayNameOfState(EGameFeaturePluginState State)
-{
-	switch (State)
-	{
-	case EGameFeaturePluginState::Uninitialized: return LOCTEXT("UninitializedStateDisplayName", "Uninitialized");
-	case EGameFeaturePluginState::UnknownStatus: return LOCTEXT("UnknownStatusStateDisplayName", "UnknownStatus");
-	case EGameFeaturePluginState::CheckingStatus: return LOCTEXT("CheckingStatusStateDisplayName", "CheckingStatus");
-	case EGameFeaturePluginState::StatusKnown: return LOCTEXT("StatusKnownStateDisplayName", "StatusKnown");
-	case EGameFeaturePluginState::Uninstalling: return LOCTEXT("UninstallingStateDisplayName", "Uninstalling");
-	case EGameFeaturePluginState::Downloading: return LOCTEXT("DownloadingStateDisplayName", "Downloading");
-	case EGameFeaturePluginState::Installed: return LOCTEXT("InstalledStateDisplayName", "Installed");
-	case EGameFeaturePluginState::Unmounting: return LOCTEXT("UnmountingStateDisplayName", "Unmounting");
-	case EGameFeaturePluginState::Mounting: return LOCTEXT("MountingStateDisplayName", "Mounting");
-	case EGameFeaturePluginState::WaitingForDependencies: return LOCTEXT("WaitingForDependenciesStateDisplayName", "WaitingForDependencies");
-	case EGameFeaturePluginState::Unregistering: return LOCTEXT("UnregisteringStateDisplayName", "Unregistering");
-	case EGameFeaturePluginState::Registering: return LOCTEXT("RegisteringStateDisplayName", "Registering");
-	case EGameFeaturePluginState::Registered: return LOCTEXT("RegisteredStateDisplayName", "Registered");
-	case EGameFeaturePluginState::Unloading: return LOCTEXT("UnloadingStateDisplayName", "Unloading");
-	case EGameFeaturePluginState::Loading: return LOCTEXT("LoadingStateDisplayName", "Loading");
-	case EGameFeaturePluginState::Loaded: return LOCTEXT("LoadedStateDisplayName", "Loaded");
-	case EGameFeaturePluginState::Deactivating: return LOCTEXT("DeactivatingStateDisplayName", "Deactivating");
-	case EGameFeaturePluginState::Activating: return LOCTEXT("ActivatingStateDisplayName", "Activating");
-	case EGameFeaturePluginState::Active: return LOCTEXT("ActiveStateDisplayName", "Active");
-	default:
-		check(0);
-		return FText::GetEmpty();
-	}
-}
 
 EGameFeaturePluginState FGameFeatureDataDetailsCustomization::GetCurrentState() const
 {
 	return UGameFeaturesSubsystem::Get().GetPluginState(PluginURL);
 }
 
-FText FGameFeatureDataDetailsCustomization::GetStateStatusDisplay() const
+FText FGameFeatureDataDetailsCustomization::GetInitialStateText() const
 {
-	// Display the current state/transition for anything but the four acceptable destination states (which are already covered by the switcher)
-	const EGameFeaturePluginState State = GetCurrentState();
-	switch (State)
-	{
-		case EGameFeaturePluginState::Active:
-		case EGameFeaturePluginState::Installed:
-		case EGameFeaturePluginState::Loaded:
-		case EGameFeaturePluginState::Registered:
-			return FText::GetEmpty();
-		default:
-			return GetDisplayNameOfState(State);
-	}
+	const EBuiltInAutoState AutoState = UGameFeaturesSubsystem::DetermineBuiltInInitialFeatureState(PluginPtr->GetDescriptor().CachedJson, FString());
+	const EGameFeaturePluginState InitialState = UGameFeaturesSubsystem::ConvertInitialFeatureStateToTargetState(AutoState);
+	return SGameFeatureStateWidget::GetDisplayNameOfState(InitialState);
 }
 
 void FGameFeatureDataDetailsCustomization::OnOperationCompletedOrFailed(const UE::GameFeatures::FResult& Result, const TWeakPtr<FGameFeatureDataDetailsCustomization> WeakThisPtr)
diff --git a/Engine/Plugins/Experimental/GameFeatures/Source/GameFeaturesEditor/Private/GameFeatureDataDetailsCustomization.h b/Engine/Plugins/Experimental/GameFeatures/Source/GameFeaturesEditor/Private/GameFeatureDataDetailsCustomization.h
--- a/Engine/Plugins/Experimental/GameFeatures/Source/GameFeaturesEditor/Private/GameFeatureDataDetailsCustomization.h
+++ b/Engine/Plugins/Experimental/GameFeatures/Source/GameFeaturesEditor/Private/GameFeatureDataDetailsCustomization.h
@@ -9,6 +9,7 @@
 
 class IDetailLayoutBuilder;
 class SErrorText;
+class IPlugin;
 
 //////////////////////////////////////////////////////////////////////////
 // FGameFeatureDataDetailsCustomization
@@ -28,16 +29,15 @@ public:
 protected:
 	void ChangeDesiredState(EGameFeaturePluginState State);
 
-	static FText GetDisplayNameOfState(EGameFeaturePluginState State);
-
 	EGameFeaturePluginState GetCurrentState() const;
 
-	FText GetStateStatusDisplay() const;
+	FText GetInitialStateText() const;
 
 	static void OnOperationCompletedOrFailed(const UE::GameFeatures::FResult& Result, const TWeakPtr<FGameFeatureDataDetailsCustomization> WeakThisPtr);
 protected:
 	TArray<TWeakObjectPtr<UObject>> ObjectsBeingCustomized;
 	FString PluginURL;
+	TSharedPtr<IPlugin> PluginPtr;
 
 	TSharedPtr<SErrorText> ErrorTextWidget;
 };
diff --git a/Engine/Plugins/Experimental/GameFeatures/Source/GameFeaturesEditor/Private/GameFeaturePluginMetadataCustomization.cpp b/Engine/Plugins/Experimental/GameFeatures/Source/GameFeaturesEditor/Private/GameFeaturePluginMetadataCustomization.cpp
new file mode 100644
--- /dev/null
+++ b/Engine/Plugins/Experimental/GameFeatures/Source/GameFeaturesEditor/Private/GameFeaturePluginMetadataCustomization.cpp
@@ -0,0 +1,89 @@
+// Copyright Epic Games, Inc. All Rights Reserved.
+
+#include "GameFeaturePluginMetadataCustomization.h"
+#include "Widgets/DeclarativeSyntaxSupport.h"
+#include "Widgets/SBoxPanel.h"
+#include "Widgets/Input/SButton.h"
+#include "Widgets/Input/SSegmentedControl.h"
+#include "Widgets/Notifications/SErrorText.h"
+#include "DetailLayoutBuilder.h"
+#include "DetailWidgetRow.h"
+#include "DetailCategoryBuilder.h"
+#include "Dom/JsonValue.h"
+#include "Features/IPluginsEditorFeature.h"
+#include "PluginDescriptor.h"
+#include "Interfaces/IPluginManager.h"
+#include "SGameFeatureStateWidget.h"
+
+#include "GameFeatureData.h"
+#include "GameFeaturesSubsystem.h"
+#include "../../GameFeatures/Private/GameFeaturePluginStateMachine.h"
+
+#define LOCTEXT_NAMESPACE "GameFeatures"
+
+//////////////////////////////////////////////////////////////////////////
+// FGameFeaturePluginMetadataCustomization
+
+void FGameFeaturePluginMetadataCustomization::CustomizeDetails(FPluginEditingContext& InPluginContext, IDetailLayoutBuilder& DetailBuilder)
+{
+	Plugin = InPluginContext.PluginBeingEdited;
+	
+	const EBuiltInAutoState AutoState = UGameFeaturesSubsystem::DetermineBuiltInInitialFeatureState(Plugin->GetDescriptor().CachedJson, FString());
+	InitialState = UGameFeaturesSubsystem::ConvertInitialFeatureStateToTargetState(AutoState);
+
+	IDetailCategoryBuilder& TopCategory = DetailBuilder.EditCategory("Game Features", FText::GetEmpty(), ECategoryPriority::Important);
+
+	FDetailWidgetRow& ControlRow = TopCategory.AddCustomRow(LOCTEXT("ControlSearchText", "Plugin State Control"))
+		.NameContent()
+		[
+			SNew(STextBlock)
+			.Text(LOCTEXT("InitialState", "Initial State"))
+			.Font(DetailBuilder.GetDetailFont())
+		]
+		.ValueContent()
+		[
+			SNew(SGameFeatureStateWidget)
+			.ToolTipText(LOCTEXT("DefaultStateSwitcherTooltip", "Change the default initial state of this game feature"))
+			.CurrentState(this, &FGameFeaturePluginMetadataCustomization::GetDefaultState)
+			.OnStateChanged(this, &FGameFeaturePluginMetadataCustomization::ChangeDefaultState)
+		];
+}
+
+void FGameFeaturePluginMetadataCustomization::CommitEdits(FPluginDescriptor& Descriptor)
+{
+	FString StateStr;
+	switch (InitialState)
+	{
+	case EGameFeaturePluginState::Installed:
+		StateStr = TEXT("Installed");
+		break;
+	case EGameFeaturePluginState::Registered:
+		StateStr = TEXT("Registered");
+		break;
+	case EGameFeaturePluginState::Loaded:
+		StateStr = TEXT("Loaded");
+		break;
+	case EGameFeaturePluginState::Active:
+		StateStr = TEXT("Active");
+		break;
+	}
+
+	if (ensure(!StateStr.IsEmpty()))
+	{
+		Descriptor.AdditionalFieldsToWrite.FindOrAdd(TEXT("BuiltInInitialFeatureState")) = MakeShared<FJsonValueString>(StateStr);
+	}
+}
+
+EGameFeaturePluginState FGameFeaturePluginMetadataCustomization::GetDefaultState() const
+{
+	return InitialState;
+}
+
+void FGameFeaturePluginMetadataCustomization::ChangeDefaultState(EGameFeaturePluginState DesiredState)
+{
+	InitialState = DesiredState;
+}
+
+//////////////////////////////////////////////////////////////////////////
+
+#undef LOCTEXT_NAMESPACE
diff --git a/Engine/Plugins/Experimental/GameFeatures/Source/GameFeaturesEditor/Private/GameFeaturePluginMetadataCustomization.h b/Engine/Plugins/Experimental/GameFeatures/Source/GameFeaturesEditor/Private/GameFeaturePluginMetadataCustomization.h
new file mode 100644
--- /dev/null
+++ b/Engine/Plugins/Experimental/GameFeatures/Source/GameFeaturesEditor/Private/GameFeaturePluginMetadataCustomization.h
@@ -0,0 +1,33 @@
+// Copyright Epic Games, Inc. All Rights Reserved.
+
+#pragma once
+
+#include "CoreMinimal.h"
+#include "Input/Reply.h"
+#include "IDetailCustomization.h"
+#include "Features/IPluginsEditorFeature.h"
+
+class IDetailLayoutBuilder;
+struct FPluginEditingContext;
+class IPlugin;
+struct FPluginDescriptor;
+enum class EGameFeaturePluginState : uint8;
+
+//////////////////////////////////////////////////////////////////////////
+// FGameFeaturePluginMetadataCustomization
+
+class FGameFeaturePluginMetadataCustomization : public FPluginEditorExtension
+{
+public:
+	void CustomizeDetails(FPluginEditingContext& InPluginContext, IDetailLayoutBuilder& DetailBuilder);
+
+	virtual void CommitEdits(FPluginDescriptor& Descriptor) override;
+private:
+	EGameFeaturePluginState GetDefaultState() const;
+
+	void ChangeDefaultState(EGameFeaturePluginState DesiredState);
+
+	TSharedPtr<IPlugin> Plugin;
+
+	EGameFeaturePluginState InitialState;
+};
diff --git a/Engine/Plugins/Experimental/GameFeatures/Source/GameFeaturesEditor/Private/GameFeaturesEditorModule.cpp b/Engine/Plugins/Experimental/GameFeatures/Source/GameFeaturesEditor/Private/GameFeaturesEditorModule.cpp
--- a/Engine/Plugins/Experimental/GameFeatures/Source/GameFeaturesEditor/Private/GameFeaturesEditorModule.cpp
+++ b/Engine/Plugins/Experimental/GameFeatures/Source/GameFeaturesEditor/Private/GameFeaturesEditorModule.cpp
@@ -13,6 +13,7 @@
 #include "ContentBrowserModule.h"
 #include "Subsystems/AssetEditorSubsystem.h"
 #include "GameFeatureDataDetailsCustomization.h"
+#include "GameFeaturePluginMetadataCustomization.h"
 #include "Logging/MessageLog.h"
 #include "SSettingsEditorCheckoutNotice.h"
 #include "Engine/AssetManagerSettings.h"
@@ -21,6 +22,7 @@
 #include "Engine/AssetManager.h"
 #include "HAL/FileManager.h"
 #include "Editor.h"
+#include "Dom/JsonValue.h"
 
 #define LOCTEXT_NAMESPACE "GameFeatures"
 
@@ -61,6 +63,13 @@ struct FGameFeaturePluginTemplateDescription : public FPluginTemplateDescription
 		FPaths::MakePlatformFilename(InOutPath);
 	}
 
+	virtual void CustomizeDescriptorBeforeCreation(FPluginDescriptor& Descriptor) override
+	{
+		Descriptor.bExplicitlyLoaded = true;
+		Descriptor.AdditionalFieldsToWrite.FindOrAdd(TEXT("BuiltInInitialFeatureState")) = MakeShared<FJsonValueString>(TEXT("Active"));
+		Descriptor.Category = TEXT("Game Features");
+	}
+
 	virtual void OnPluginCreated(TSharedPtr<IPlugin> NewPlugin) override
 	{
 		TSubclassOf<UGameFeatureData> DefaultGameFeatureDataClass = GetDefault<UGameFeaturesSubsystemSettings>()->DefaultGameFeatureDataClass;
@@ -137,28 +146,11 @@ class FGameFeaturesEditorModule : public FDefaultModuleImpl
 
 			if (ModularFeatures.IsModularFeatureAvailable(EditorFeatures::PluginsEditor))
 			{
-				IPluginsEditorFeature& PluginsEditor = ModularFeatures.GetModularFeature<IPluginsEditorFeature>(EditorFeatures::PluginsEditor);
-				PluginsEditor.RegisterPluginTemplate(ContentOnlyTemplate.ToSharedRef());
+				OnModularFeatureRegistered(EditorFeatures::PluginsEditor, &ModularFeatures.GetModularFeature<IPluginsEditorFeature>(EditorFeatures::PluginsEditor));
 			}
 		}
 	}
 
-	void OnModularFeatureRegistered(const FName& Type, class IModularFeature* ModularFeature)
-	{
-		if (Type == EditorFeatures::PluginsEditor)
-		{
-			static_cast<IPluginsEditorFeature*>(ModularFeature)->RegisterPluginTemplate(ContentOnlyTemplate.ToSharedRef());
-		}
-	}
-
-	void OnModularFeatureUnregistered(const FName& Type, class IModularFeature* ModularFeature)
-	{
-		if (Type == EditorFeatures::PluginsEditor)
-		{
-			static_cast<IPluginsEditorFeature*>(ModularFeature)->UnregisterPluginTemplate(ContentOnlyTemplate.ToSharedRef());
-		}
-	}
-
 	virtual void ShutdownModule() override
 	{
 		// Remove the customization
@@ -178,13 +170,33 @@ class FGameFeaturesEditorModule : public FDefaultModuleImpl
 
 			if (ModularFeatures.IsModularFeatureAvailable(EditorFeatures::PluginsEditor))
 			{
-				IPluginsEditorFeature& PluginsEditor = ModularFeatures.GetModularFeature<IPluginsEditorFeature>(EditorFeatures::PluginsEditor);
-				PluginsEditor.UnregisterPluginTemplate(ContentOnlyTemplate.ToSharedRef());
+				OnModularFeatureUnregistered(EditorFeatures::PluginsEditor, &ModularFeatures.GetModularFeature<IPluginsEditorFeature>(EditorFeatures::PluginsEditor));
 			}
 			ContentOnlyTemplate.Reset();
 		}
 	}
 
+
+	void OnModularFeatureRegistered(const FName& Type, class IModularFeature* ModularFeature)
+	{
+		if (Type == EditorFeatures::PluginsEditor)
+		{
+			IPluginsEditorFeature& PluginEditor = *static_cast<IPluginsEditorFeature*>(ModularFeature);
+			PluginEditor.RegisterPluginTemplate(ContentOnlyTemplate.ToSharedRef());
+			PluginEditorExtensionDelegate = PluginEditor.RegisterPluginEditorExtension(FOnPluginBeingEdited::CreateRaw(this, &FGameFeaturesEditorModule::CustomizePluginEditing));
+		}
+	}
+
+	void OnModularFeatureUnregistered(const FName& Type, class IModularFeature* ModularFeature)
+	{
+		if (Type == EditorFeatures::PluginsEditor)
+		{
+			IPluginsEditorFeature& PluginEditor = *static_cast<IPluginsEditorFeature*>(ModularFeature);
+			PluginEditor.UnregisterPluginTemplate(ContentOnlyTemplate.ToSharedRef());
+			PluginEditor.UnregisterPluginEditorExtension(PluginEditorExtensionDelegate);
+		}
+	}
+
 	void AddDefaultGameDataRule()
 	{
 		// Check out the ini or make it writable
@@ -269,8 +281,21 @@ class FGameFeaturesEditorModule : public FDefaultModuleImpl
 		}
 	}
 
+	TSharedPtr<FPluginEditorExtension> CustomizePluginEditing(FPluginEditingContext& InPluginContext, IDetailLayoutBuilder& DetailBuilder)
+	{
+		const bool bIsGameFeaturePlugin = InPluginContext.PluginBeingEdited->GetDescriptorFileName().Contains(TEXT("/GameFeatures/"));
+		if (bIsGameFeaturePlugin)
+		{
+			TSharedPtr<FGameFeaturePluginMetadataCustomization> Result = MakeShareable(new FGameFeaturePluginMetadataCustomization);
+			Result->CustomizeDetails(InPluginContext, DetailBuilder);
+			return Result;
+		}
+
+		return nullptr;
+	}
 private:
 	TSharedPtr<FPluginTemplateDescription> ContentOnlyTemplate;
+	FPluginEditorExtensionHandle PluginEditorExtensionDelegate;
 };
 
 IMPLEMENT_MODULE(FGameFeaturesEditorModule, GameFeaturesEditor)
diff --git a/Engine/Plugins/Experimental/GameFeatures/Source/GameFeaturesEditor/Private/SGameFeatureStateWidget.cpp b/Engine/Plugins/Experimental/GameFeatures/Source/GameFeaturesEditor/Private/SGameFeatureStateWidget.cpp
new file mode 100644
--- /dev/null
+++ b/Engine/Plugins/Experimental/GameFeatures/Source/GameFeaturesEditor/Private/SGameFeatureStateWidget.cpp
@@ -0,0 +1,99 @@
+// Copyright Epic Games, Inc. All Rights Reserved.
+
+#include "SGameFeatureStateWidget.h"
+#include "Widgets/DeclarativeSyntaxSupport.h"
+#include "Widgets/SBoxPanel.h"
+#include "Widgets/Input/SSegmentedControl.h"
+#include "GameFeatureData.h"
+#include "../../GameFeatures/Private/GameFeaturePluginStateMachine.h"
+
+#define LOCTEXT_NAMESPACE "GameFeatures"
+
+//////////////////////////////////////////////////////////////////////////
+// FGameFeatureDataDetailsCustomization
+
+void SGameFeatureStateWidget::Construct(const FArguments& InArgs)
+{
+	CurrentState = InArgs._CurrentState;
+
+	ChildSlot
+	[
+		SNew(SHorizontalBox)
+		+SHorizontalBox::Slot()
+		.AutoWidth()
+		.VAlign(VAlign_Center)
+		[
+			SNew(SSegmentedControl<EGameFeaturePluginState>)
+			.Value(CurrentState)
+			.OnValueChanged(InArgs._OnStateChanged)
+//			.ToolTipText(LOCTEXT("StateSwitcherTooltip", "Attempt to change the current state of this game feature"))
+			+SSegmentedControl<EGameFeaturePluginState>::Slot(EGameFeaturePluginState::Installed)
+				.Text(GetDisplayNameOfState(EGameFeaturePluginState::Installed))
+			+ SSegmentedControl<EGameFeaturePluginState>::Slot(EGameFeaturePluginState::Registered)
+				.Text(GetDisplayNameOfState(EGameFeaturePluginState::Registered))
+			+ SSegmentedControl<EGameFeaturePluginState>::Slot(EGameFeaturePluginState::Loaded)
+				.Text(GetDisplayNameOfState(EGameFeaturePluginState::Loaded))
+			+ SSegmentedControl<EGameFeaturePluginState>::Slot(EGameFeaturePluginState::Active)
+				.Text(GetDisplayNameOfState(EGameFeaturePluginState::Active))
+		]
+		+SHorizontalBox::Slot()
+		.Padding(8.0f, 0.0f, 0.0f, 0.0f)
+		.VAlign(VAlign_Center)
+		[
+			SNew(STextBlock)
+			.Text(this, &SGameFeatureStateWidget::GetStateStatusDisplay)
+			.TextStyle(&FAppStyle::Get().GetWidgetStyle<FTextBlockStyle>("ButtonText"))
+//			.ToolTipText(LOCTEXT("OtherStateToolTip", "The current state of this game feature plugin"))
+			.ColorAndOpacity(FAppStyle::Get().GetSlateColor(TEXT("Colors.AccentYellow")))
+		]
+	];
+}
+
+FText SGameFeatureStateWidget::GetDisplayNameOfState(EGameFeaturePluginState State)
+{
+	switch (State)
+	{
+	case EGameFeaturePluginState::Uninitialized: return LOCTEXT("UninitializedStateDisplayName", "Uninitialized");
+	case EGameFeaturePluginState::UnknownStatus: return LOCTEXT("UnknownStatusStateDisplayName", "UnknownStatus");
+	case EGameFeaturePluginState::CheckingStatus: return LOCTEXT("CheckingStatusStateDisplayName", "CheckingStatus");
+	case EGameFeaturePluginState::StatusKnown: return LOCTEXT("StatusKnownStateDisplayName", "StatusKnown");
+	case EGameFeaturePluginState::Uninstalling: return LOCTEXT("UninstallingStateDisplayName", "Uninstalling");
+	case EGameFeaturePluginState::Downloading: return LOCTEXT("DownloadingStateDisplayName", "Downloading");
+	case EGameFeaturePluginState::Installed: return LOCTEXT("InstalledStateDisplayName", "Installed");
+	case EGameFeaturePluginState::Unmounting: return LOCTEXT("UnmountingStateDisplayName", "Unmounting");
+	case EGameFeaturePluginState::Mounting: return LOCTEXT("MountingStateDisplayName", "Mounting");
+	case EGameFeaturePluginState::WaitingForDependencies: return LOCTEXT("WaitingForDependenciesStateDisplayName", "WaitingForDependencies");
+	case EGameFeaturePluginState::Unregistering: return LOCTEXT("UnregisteringStateDisplayName", "Unregistering");
+	case EGameFeaturePluginState::Registering: return LOCTEXT("RegisteringStateDisplayName", "Registering");
+	case EGameFeaturePluginState::Registered: return LOCTEXT("RegisteredStateDisplayName", "Registered");
+	case EGameFeaturePluginState::Unloading: return LOCTEXT("UnloadingStateDisplayName", "Unloading");
+	case EGameFeaturePluginState::Loading: return LOCTEXT("LoadingStateDisplayName", "Loading");
+	case EGameFeaturePluginState::Loaded: return LOCTEXT("LoadedStateDisplayName", "Loaded");
+	case EGameFeaturePluginState::Deactivating: return LOCTEXT("DeactivatingStateDisplayName", "Deactivating");
+	case EGameFeaturePluginState::Activating: return LOCTEXT("ActivatingStateDisplayName", "Activating");
+	case EGameFeaturePluginState::Active: return LOCTEXT("ActiveStateDisplayName", "Active");
+	default:
+		check(0);
+		return FText::GetEmpty();
+	}
+}
+
+FText SGameFeatureStateWidget::GetStateStatusDisplay() const
+{
+	// Display the current state/transition for anything but the four acceptable destination states (which are already covered by the switcher)
+	const EGameFeaturePluginState State = CurrentState.Get();
+	switch (State)
+	{
+		case EGameFeaturePluginState::Active:
+		case EGameFeaturePluginState::Installed:
+		case EGameFeaturePluginState::Loaded:
+		case EGameFeaturePluginState::Registered:
+			return FText::GetEmpty();
+		default:
+			return GetDisplayNameOfState(State);
+	}
+}
+
+//////////////////////////////////////////////////////////////////////////
+
+#undef LOCTEXT_NAMESPACE
diff --git a/Engine/Plugins/Experimental/GameFeatures/Source/GameFeaturesEditor/Private/SGameFeatureStateWidget.h b/Engine/Plugins/Experimental/GameFeatures/Source/GameFeaturesEditor/Private/SGameFeatureStateWidget.h
new file mode 100644
--- /dev/null
+++ b/Engine/Plugins/Experimental/GameFeatures/Source/GameFeaturesEditor/Private/SGameFeatureStateWidget.h
@@ -0,0 +1,41 @@
+// Copyright Epic Games, Inc. All Rights Reserved.
+
+#pragma once
+
+#include "CoreMinimal.h"
+#include "Input/Reply.h"
+#include "Widgets/SCompoundWidget.h"
+#include "GameFeaturesSubsystem.h"
+
+enum class EGameFeaturePluginState : uint8;
+
+//////////////////////////////////////////////////////////////////////////
+// FGameFeatureDataDetailsCustomization
+
+/**
+ * A delegate that is invoked when widgets want to notify a user that they have been clicked.
+ * Intended for use by buttons and other button-like widgets.
+ */
+DECLARE_DELEGATE_OneParam(FOnWidgetChangesGameFeatureState, EGameFeaturePluginState)
+
+class SGameFeatureStateWidget : public SCompoundWidget
+{
+public:
+	SLATE_BEGIN_ARGS(SGameFeatureStateWidget) {}
+		SLATE_ATTRIBUTE(EGameFeaturePluginState, CurrentState)
+		SLATE_EVENT(FOnWidgetChangesGameFeatureState, OnStateChanged)
+	SLATE_END_ARGS()
+
+public:
+	SGameFeatureStateWidget() {}
+
+	void Construct(const FArguments& InArgs);
+
+	static FText GetDisplayNameOfState(EGameFeaturePluginState State);
+
+private:
+	FText GetStateStatusDisplay() const;
+
+private:
+	TAttribute<EGameFeaturePluginState> CurrentState;
+};
diff --git a/Engine/Source/Editor/UnrealEd/Public/Features/IPluginsEditorFeature.h b/Engine/Source/Editor/UnrealEd/Public/Features/IPluginsEditorFeature.h
--- a/Engine/Source/Editor/UnrealEd/Public/Features/IPluginsEditorFeature.h
+++ b/Engine/Source/Editor/UnrealEd/Public/Features/IPluginsEditorFeature.h
@@ -8,6 +8,9 @@
 
 struct FSlateDynamicImageBrush;
 class IPlugin;
+class IDetailLayoutBuilder;
+struct FPluginDescriptor;
+class SWidget;
 
 /**
  * Description of a plugin template
@@ -41,6 +44,9 @@ struct FPluginTemplateDescription
 	/** What is the expected Loading Phase for this plugin? */
 	ELoadingPhase::Type LoadingPhase;
 
+	/** Called just before the plugin is created */
+	virtual void CustomizeDescriptorBeforeCreation(FPluginDescriptor& Descriptor) {}
+
 	/** Called after the plugin has been created */
 	virtual void OnPluginCreated(TSharedPtr<IPlugin> NewPlugin) {}
 
@@ -71,6 +77,26 @@ struct FPluginTemplateDescription
 	virtual ~FPluginTemplateDescription() {}
 };
 
+struct FPluginEditingContext
+{
+	TSharedPtr<IPlugin> PluginBeingEdited;
+
+};
+
+/**
+ * Description of a plugin editor extension
+ */
+struct FPluginEditorExtension : public TSharedFromThis<FPluginEditorExtension>
+{
+	virtual ~FPluginEditorExtension() {}
+	virtual void CommitEdits(FPluginDescriptor& Descriptor) = 0;
+};
+
+/** Notification when editing a plugin to allow further customization */
+DECLARE_DELEGATE_RetVal_TwoParams(TSharedPtr<FPluginEditorExtension>, FOnPluginBeingEdited, FPluginEditingContext& /*InPluginContext*/, IDetailLayoutBuilder& /*DetailBuilder*/);
+
+using FPluginEditorExtensionHandle = int32;
+
 /**
  * Feature interface for a Plugins management UI
  */
@@ -87,5 +113,20 @@ public:
 	  * Unregisters the specified plugin template from the New Plugin wizard
 	  */
 	virtual void UnregisterPluginTemplate(TSharedRef<FPluginTemplateDescription> Template) = 0;
+
+	/**
+	  * Further customize the detail panel experience when editing a specific plugin
+	  */
+	virtual FPluginEditorExtensionHandle RegisterPluginEditorExtension(FOnPluginBeingEdited Extension) = 0;
+
+	/**
+	  * Remove a customization
+	  */
+	virtual void UnregisterPluginEditorExtension(FPluginEditorExtensionHandle ExtensionHandle) = 0;
+
+	/**
+	 * Open the plugin editor for a plugin
+	 */
+	virtual void OpenPluginEditor(TSharedRef<IPlugin> PluginToEdit, TSharedPtr<SWidget> ParentWidget, FSimpleDelegate OnEditCommitted) = 0;
 };
 
diff --git a/Engine/Source/Runtime/Projects/Private/PluginDescriptor.cpp b/Engine/Source/Runtime/Projects/Private/PluginDescriptor.cpp
--- a/Engine/Source/Runtime/Projects/Private/PluginDescriptor.cpp
+++ b/Engine/Source/Runtime/Projects/Private/PluginDescriptor.cpp
@@ -80,6 +80,7 @@ bool FPluginDescriptor::Load(const FString& FileName, FText& OutFailReason)
 {
 #if WITH_EDITOR
 	CachedJson.Reset();
+	AdditionalFieldsToWrite.Reset();
 #endif // WITH_EDITOR
 
 	FString Text;
@@ -95,6 +96,7 @@ bool FPluginDescriptor::Read(const FString& Text, FText& OutFailReason)
 {
 #if WITH_EDITOR
 	CachedJson.Reset();
+	AdditionalFieldsToWrite.Reset();
 #endif // WITH_EDITOR
 
 	// Deserialize a JSON object from the string
@@ -106,6 +108,7 @@ bool FPluginDescriptor::Read(const FString& Text, FText& OutFailReason)
 		{
 #if WITH_EDITOR
 			CachedJson = JsonObject;
+			AdditionalFieldsToWrite.Reset();
 #endif // WITH_EDITOR
 			return true;
 		}
@@ -229,10 +232,18 @@ void FPluginDescriptor::Write(FString& Text) const
 
 void FPluginDescriptor::Write(TJsonWriter<>& Writer) const
 {
-	TSharedRef<FJsonObject> PluginJsonObject = MakeShared<FJsonObject>();
+	TSharedPtr<FJsonObject> PluginJsonObject = MakeShared<FJsonObject>();
+
+#if WITH_EDITOR
+	if (CachedJson.IsValid())
+	{
+		FJsonObject::Duplicate(/*Source=*/ CachedJson, /*Dest=*/ PluginJsonObject);
+	}
+#endif
+
 	UpdateJson(*PluginJsonObject);
 
-	FJsonSerializer::Serialize(PluginJsonObject, Writer);
+	FJsonSerializer::Serialize(PluginJsonObject.ToSharedRef(), Writer);
 }
 
 void FPluginDescriptor::UpdateJson(FJsonObject& JsonObject) const
@@ -353,6 +364,13 @@ void FPluginDescriptor::UpdateJson(FJsonObject& JsonObject) const
 	PostBuildSteps.UpdateJson(JsonObject, TEXT("PostBuildSteps"));
 
 	FPluginReferenceDescriptor::UpdateArray(JsonObject, TEXT("Plugins"), Plugins);
+
+#if WITH_EDITOR
+	for (const auto& KVP : AdditionalFieldsToWrite)
+	{
+		JsonObject.SetField(KVP.Key, FJsonValue::Duplicate(KVP.Value));
+	}
+#endif
 }
 
 bool FPluginDescriptor::UpdatePluginFile(const FString& FileName, FText& OutFailReason) const
@@ -384,6 +402,9 @@ bool FPluginDescriptor::UpdatePluginFile(const FString& FileName, FText& OutFail
 			}
 		}
 
+#if WITH_EDITOR
+		CachedJson = JsonObject;
+#endif
 		return PluginDescriptor::WriteFile(FileName, JsonText, OutFailReason);
 	}
 	else
diff --git a/Engine/Source/Runtime/Projects/Public/PluginDescriptor.h b/Engine/Source/Runtime/Projects/Public/PluginDescriptor.h
--- a/Engine/Source/Runtime/Projects/Public/PluginDescriptor.h
+++ b/Engine/Source/Runtime/Projects/Public/PluginDescriptor.h
@@ -9,6 +9,7 @@
 #include "PluginReferenceDescriptor.h"
 
 class FJsonObject;
+class FJsonValue;
 
 /**
  * Setting for whether a plugin is enabled by default
@@ -118,7 +119,10 @@ struct PROJECTS_API FPluginDescriptor
 
 #if WITH_EDITOR
 	/** Cached json for custom data */
-	TSharedPtr<FJsonObject> CachedJson;
+	mutable TSharedPtr<FJsonObject> CachedJson;
+
+	/** Additional fields to write */
+	TMap<FString, TSharedPtr<FJsonValue>> AdditionalFieldsToWrite;
 #endif
 
 	/** Constructor. */
