Fixing non-unity missing includes in GameFeatureDataDetailsCustomization.cpp

diff --git a/Engine/Plugins/Experimental/GameFeatures/Source/GameFeaturesEditor/Private/GameFeatureDataDetailsCustomization.cpp b/Engine/Plugins/Experimental/GameFeatures/Source/GameFeaturesEditor/Private/GameFeatureDataDetailsCustomization.cpp
--- a/Engine/Plugins/Experimental/GameFeatures/Source/GameFeaturesEditor/Private/GameFeatureDataDetailsCustomization.cpp
+++ b/Engine/Plugins/Experimental/GameFeatures/Source/GameFeaturesEditor/Private/GameFeatureDataDetailsCustomization.cpp
@@ -4,6 +4,7 @@
 #include "Widgets/DeclarativeSyntaxSupport.h"
 #include "Widgets/SBoxPanel.h"
 #include "Widgets/Input/SButton.h"
+#include "Widgets/Images/SImage.h"
 #include "SGameFeatureStateWidget.h"
 #include "Widgets/Notifications/SErrorText.h"
 #include "DetailLayoutBuilder.h"
@@ -13,6 +14,8 @@
 #include "Interfaces/IPluginManager.h"
 #include "Features/IPluginsEditorFeature.h"
 #include "Features/EditorFeatures.h"
+#include "Features/IModularFeatures.h"
+#include "Misc/MessageDialog.h"
 
 #include "GameFeatureData.h"
 #include "../../GameFeatures/Private/GameFeaturePluginStateMachine.h"
