[NO NON-PII COMMIT MESSAGE]

diff --git a/Engine/Plugins/Experimental/GameFeatures/Source/GameFeatures/Private/GameFeaturesProjectPolicies.cpp b/Engine/Plugins/Experimental/GameFeatures/Source/GameFeatures/Private/GameFeaturesProjectPolicies.cpp
--- a/Engine/Plugins/Experimental/GameFeatures/Source/GameFeatures/Private/GameFeaturesProjectPolicies.cpp
+++ b/Engine/Plugins/Experimental/GameFeatures/Source/GameFeatures/Private/GameFeaturesProjectPolicies.cpp
@@ -2,6 +2,7 @@
 
 #include "GameFeaturesProjectPolicies.h"
 #include "GameFeaturesSubsystem.h"
+#include "GameFeaturesSubsystemSettings.h"
 #include "Misc/CoreMisc.h"
 
 void UDefaultGameFeaturesProjectPolicies::InitGameFeatureManager()
@@ -22,3 +23,21 @@ void UDefaultGameFeaturesProjectPolicies::GetGameFeatureLoadingMode(bool& bLoadC
 	bLoadClientData = !IsRunningDedicatedServer();
 	bLoadServerData = !IsRunningClientOnly();
 }
+
+const TArray<FName> UDefaultGameFeaturesProjectPolicies::GetPreloadBundleStateForGameFeature() const
+{
+	// By default, use the bundles corresponding to loading mode
+	bool bLoadClientData, bLoadServerData;
+	GetGameFeatureLoadingMode(bLoadClientData, bLoadServerData);
+
+	TArray<FName> FeatureBundles;
+	if (bLoadClientData)
+	{
+		FeatureBundles.Add(UGameFeaturesSubsystemSettings::LoadStateClient);
+	}
+	if (bLoadServerData)
+	{
+		FeatureBundles.Add(UGameFeaturesSubsystemSettings::LoadStateServer);
+	}
+	return FeatureBundles;
+}
diff --git a/Engine/Plugins/Experimental/GameFeatures/Source/GameFeatures/Public/GameFeaturesProjectPolicies.h b/Engine/Plugins/Experimental/GameFeatures/Source/GameFeatures/Public/GameFeaturesProjectPolicies.h
--- a/Engine/Plugins/Experimental/GameFeatures/Source/GameFeatures/Public/GameFeaturesProjectPolicies.h
+++ b/Engine/Plugins/Experimental/GameFeatures/Source/GameFeatures/Public/GameFeaturesProjectPolicies.h
@@ -52,5 +52,6 @@ public:
 	//~UGameFeaturesProjectPolicies interface
 	virtual void InitGameFeatureManager() override;
 	virtual void GetGameFeatureLoadingMode(bool& bLoadClientData, bool& bLoadServerData) const override;
+	virtual const TArray<FName> GetPreloadBundleStateForGameFeature() const override;
 	//~End of UGameFeaturesProjectPolicies interface
 };
