[NO NON-PII COMMIT MESSAGE]

diff --git a/Engine/Plugins/Experimental/GameFeatures/Source/GameFeaturesEditor/Private/GameFeatureDataDetailsCustomization.cpp b/Engine/Plugins/Experimental/GameFeatures/Source/GameFeaturesEditor/Private/GameFeatureDataDetailsCustomization.cpp
--- a/Engine/Plugins/Experimental/GameFeatures/Source/GameFeaturesEditor/Private/GameFeatureDataDetailsCustomization.cpp
+++ b/Engine/Plugins/Experimental/GameFeatures/Source/GameFeaturesEditor/Private/GameFeatureDataDetailsCustomization.cpp
@@ -16,6 +16,7 @@
 #include "Features/EditorFeatures.h"
 #include "Features/IModularFeatures.h"
 #include "Misc/MessageDialog.h"
+#include "Misc/Paths.h"
 
 #include "GameFeatureData.h"
 #include "../../GameFeatures/Private/GameFeaturePluginStateMachine.h"
@@ -158,6 +159,22 @@ void FGameFeatureDataDetailsCustomization::CustomizeDetails(IDetailLayoutBuilder
 				]
 			];
 
+		FDetailWidgetRow& TagsRow = TopCategory.AddCustomRow(LOCTEXT("TagSearchText", "Gameplay Tag Config Path"))
+			.NameContent()
+			[
+				SNew(STextBlock)
+				.Text(LOCTEXT("TagConfigPath", "Gameplay Tag Config Path"))
+				.ToolTipText(LOCTEXT("TagConfigPathTooltip", "Path to search for Gameplay Tag ini files. To create feature-specific tags use Add New Tag Source with this path and then Add New Gameplay Tag with that Tag Source."))
+				.Font(DetailBuilder.GetDetailFont())
+			]
+			.ValueContent()
+			.MinDesiredWidth(400.0f)
+			[
+				SNew(STextBlock)
+				.Text(this, &FGameFeatureDataDetailsCustomization::GetTagConfigPathText)
+				.Font(DetailBuilder.GetDetailFont())
+			];
+
 
 //@TODO: This disables the mode switcher widget too (and it's a const cast hack...)
 // 		if (IDetailsView* ConstHackDetailsView = const_cast<IDetailsView*>(DetailBuilder.GetDetailsView()))
@@ -223,6 +240,18 @@ FText FGameFeatureDataDetailsCustomization::GetInitialStateText() const
 	return SGameFeatureStateWidget::GetDisplayNameOfState(InitialState);
 }
 
+FText FGameFeatureDataDetailsCustomization::GetTagConfigPathText() const
+{
+	FString PluginFile = UGameFeaturesSubsystem::Get().GetPluginFilenameFromPluginURL(PluginURL);
+	FString PluginFolder = FPaths::GetPath(PluginFile);
+	FString TagFolder = PluginFolder / TEXT("Config") / TEXT("Tags");
+	if (FPaths::IsUnderDirectory(TagFolder, FPaths::ProjectDir()))
+	{
+		FPaths::MakePathRelativeTo(TagFolder, *FPaths::ProjectDir());
+	}
+	return FText::AsCultureInvariant(TagFolder);
+}
+
 void FGameFeatureDataDetailsCustomization::OnOperationCompletedOrFailed(const UE::GameFeatures::FResult& Result, const TWeakPtr<FGameFeatureDataDetailsCustomization> WeakThisPtr)
 {
 	if (Result.HasError())
diff --git a/Engine/Plugins/Experimental/GameFeatures/Source/GameFeaturesEditor/Private/GameFeatureDataDetailsCustomization.h b/Engine/Plugins/Experimental/GameFeatures/Source/GameFeaturesEditor/Private/GameFeatureDataDetailsCustomization.h
--- a/Engine/Plugins/Experimental/GameFeatures/Source/GameFeaturesEditor/Private/GameFeatureDataDetailsCustomization.h
+++ b/Engine/Plugins/Experimental/GameFeatures/Source/GameFeaturesEditor/Private/GameFeatureDataDetailsCustomization.h
@@ -32,6 +32,7 @@ protected:
 	EGameFeaturePluginState GetCurrentState() const;
 
 	FText GetInitialStateText() const;
+	FText GetTagConfigPathText() const;
 
 	static void OnOperationCompletedOrFailed(const UE::GameFeatures::FResult& Result, const TWeakPtr<FGameFeatureDataDetailsCustomization> WeakThisPtr);
 protected:
