Game Features: Allow other systems to create a FGameFeatureDeactivatingContext

diff --git a/Engine/Plugins/Experimental/GameFeatures/Source/GameFeatures/Public/GameFeaturesSubsystem.h b/Engine/Plugins/Experimental/GameFeatures/Source/GameFeatures/Public/GameFeaturesSubsystem.h
--- a/Engine/Plugins/Experimental/GameFeatures/Source/GameFeatures/Public/GameFeaturesSubsystem.h
+++ b/Engine/Plugins/Experimental/GameFeatures/Source/GameFeatures/Public/GameFeaturesSubsystem.h
@@ -29,12 +29,13 @@ public:
 		return CompletionDelegate;
 	}
 
-private:
 	FGameFeatureDeactivatingContext(FSimpleDelegate&& InCompletionDelegate)
 		: CompletionDelegate(MoveTemp(InCompletionDelegate))
 	{
 	}
 
+	int32 GetNumPausers() const { return NumPausers; }
+private:
 	FSimpleDelegate CompletionDelegate;
 	int32 NumPausers = 0;
 
