Fix check firing on game feature plugin initial creation due not handling a certain state

diff --git a/Engine/Plugins/Experimental/GameFeatures/Source/GameFeaturesEditor/Private/SGameFeatureStateWidget.cpp b/Engine/Plugins/Experimental/GameFeatures/Source/GameFeaturesEditor/Private/SGameFeatureStateWidget.cpp
--- a/Engine/Plugins/Experimental/GameFeatures/Source/GameFeaturesEditor/Private/SGameFeatureStateWidget.cpp
+++ b/Engine/Plugins/Experimental/GameFeatures/Source/GameFeaturesEditor/Private/SGameFeatureStateWidget.cpp
@@ -72,6 +72,7 @@ FText SGameFeatureStateWidget::GetDisplayNameOfState(EGameFeaturePluginState Sta
 	case EGameFeaturePluginState::Deactivating: return LOCTEXT("DeactivatingStateDisplayName", "Deactivating");
 	case EGameFeaturePluginState::Activating: return LOCTEXT("ActivatingStateDisplayName", "Activating");
 	case EGameFeaturePluginState::Active: return LOCTEXT("ActiveStateDisplayName", "Active");
+	case EGameFeaturePluginState::ErrorRegistering: return LOCTEXT("ErrorRegisteringDisplayName", "ErrorRegistering");
 	default:
 		check(0);
 		return FText::GetEmpty();
