Game Features: Better support for different world modes in the Add Components action (creating both client and server components in standalone, etc...)

diff --git a/Engine/Plugins/Experimental/GameFeatures/Source/GameFeatures/Private/GameFeatureAction_AddComponents.cpp b/Engine/Plugins/Experimental/GameFeatures/Source/GameFeatures/Private/GameFeatureAction_AddComponents.cpp
--- a/Engine/Plugins/Experimental/GameFeatures/Source/GameFeatures/Private/GameFeatureAction_AddComponents.cpp
+++ b/Engine/Plugins/Experimental/GameFeatures/Source/GameFeatures/Private/GameFeatureAction_AddComponents.cpp
@@ -90,17 +90,12 @@ void UGameFeatureAction_AddComponents::AddToWorld(const FWorldContext& WorldCont
 	{
 		if (UGameFrameworkComponentManager* GFCM = UGameInstance::GetSubsystem<UGameFrameworkComponentManager>(GameInstance))
 		{
-			UE_LOG(LogGameFeatures, Verbose, TEXT("Adding components for %s to world %s"), *GetPathNameSafe(this), *World->GetDebugDisplayName());
-			
-			bool bIsServer = IsRunningDedicatedServer();
-#if WITH_EDITOR
-			checkSlow(GameInstance->GetWorldContext());
-			bIsServer |= GameInstance->GetWorldContext()->RunAsDedicated;
-#endif
-
-			//@TODO: GameFeaturePluginEnginePush: RIP listen servers (don't intend to add support for them, but we should surface that and warn if the world is NM_ListenServer or something like that)
-			const bool bIsClient = !bIsServer;
+			const ENetMode NetMode = World->GetNetMode();
+			const bool bIsServer = NetMode != NM_Client;
+			const bool bIsClient = NetMode != NM_DedicatedServer;
 
+			UE_LOG(LogGameFeatures, Verbose, TEXT("Adding components for %s to world %s (client: %d, server: %d)"), *GetPathNameSafe(this), *World->GetDebugDisplayName(), bIsClient ? 1 : 0, bIsServer ? 1 : 0);
+			
 			for (const FGameFeatureComponentEntry& Entry : ComponentList)
 			{
 				const bool bShouldAddRequest = (bIsServer && Entry.bServerComponent) || (bIsClient && Entry.bClientComponent);
