deprecated FTicker and family and replaced by thread-safe FTSTicker

diff --git a/Engine/Plugins/Animation/GameplayInsights/Source/GameplayInsights/Private/GameplayInsightsModule.cpp b/Engine/Plugins/Animation/GameplayInsights/Source/GameplayInsights/Private/GameplayInsightsModule.cpp
--- a/Engine/Plugins/Animation/GameplayInsights/Source/GameplayInsights/Private/GameplayInsightsModule.cpp
+++ b/Engine/Plugins/Animation/GameplayInsights/Source/GameplayInsights/Private/GameplayInsightsModule.cpp
@@ -47,7 +47,7 @@ void FGameplayInsightsModule::StartupModule()
 	IModularFeatures::Get().RegisterModularFeature(TraceServices::ModuleFeatureName, &GameplayTraceModule);
 	IModularFeatures::Get().RegisterModularFeature(Insights::TimingViewExtenderFeatureName, &GameplayTimingViewExtender);
 
-	TickerHandle = FTicker::GetCoreTicker().AddTicker(TEXT("GameplayInsights"), 0.0f, [this](float DeltaTime)
+	TickerHandle = FTSTicker::GetCoreTicker().AddTicker(TEXT("GameplayInsights"), 0.0f, [this](float DeltaTime)
 	{
 		QUICK_SCOPE_CYCLE_COUNTER(STAT_FGameplayInsightsModule_TickVisualizers);
 
@@ -231,7 +231,7 @@ void FGameplayInsightsModule::ShutdownModule()
 	TimingProfilerLayoutExtension.RemoveAll(this);
 #endif
 
-	FTicker::GetCoreTicker().RemoveTicker(TickerHandle);
+	FTSTicker::GetCoreTicker().RemoveTicker(TickerHandle);
 
 	IModularFeatures::Get().UnregisterModularFeature(TraceServices::ModuleFeatureName, &GameplayTraceModule);
 	IModularFeatures::Get().UnregisterModularFeature(Insights::TimingViewExtenderFeatureName, &GameplayTimingViewExtender);
diff --git a/Engine/Plugins/Animation/GameplayInsights/Source/GameplayInsights/Private/GameplayInsightsModule.h b/Engine/Plugins/Animation/GameplayInsights/Source/GameplayInsights/Private/GameplayInsightsModule.h
--- a/Engine/Plugins/Animation/GameplayInsights/Source/GameplayInsights/Private/GameplayInsightsModule.h
+++ b/Engine/Plugins/Animation/GameplayInsights/Source/GameplayInsights/Private/GameplayInsightsModule.h
@@ -49,7 +49,7 @@ private:
 
 	FGameplayTimingViewExtender GameplayTimingViewExtender;
 
-	FDelegateHandle TickerHandle;
+	FTSTicker::FDelegateHandle TickerHandle;
 
 #if WITH_EDITOR
 	FDelegateHandle CustomDebugObjectHandle;
diff --git a/Engine/Plugins/Animation/GameplayInsights/Source/GameplayInsights/Private/ObjectPropertyTrace.cpp b/Engine/Plugins/Animation/GameplayInsights/Source/GameplayInsights/Private/ObjectPropertyTrace.cpp
--- a/Engine/Plugins/Animation/GameplayInsights/Source/GameplayInsights/Private/ObjectPropertyTrace.cpp
+++ b/Engine/Plugins/Animation/GameplayInsights/Source/GameplayInsights/Private/ObjectPropertyTrace.cpp
@@ -36,7 +36,7 @@ UE_TRACE_EVENT_END()
 
 namespace ObjectPropertyTrace
 {
-	static FDelegateHandle TickerHandle;
+	static FTSTicker::FDelegateHandle TickerHandle;
 	static TArray<TWeakObjectPtr<const UObject>> Objects;
 
 	static uint32 CurrentClassPropertyStringId = 0;
@@ -219,7 +219,7 @@ namespace ObjectPropertyTrace
 void FObjectPropertyTrace::Init()
 {
 	check(!ObjectPropertyTrace::TickerHandle.IsValid());
-	ObjectPropertyTrace::TickerHandle = FTicker::GetCoreTicker().AddTicker(TEXT("ObjectPropertyTrace"), 0.0f, [](float InDelta)
+	ObjectPropertyTrace::TickerHandle = FTSTicker::GetCoreTicker().AddTicker(TEXT("ObjectPropertyTrace"), 0.0f, [](float InDelta)
 	{
 		if(UE_TRACE_CHANNELEXPR_IS_ENABLED(ObjectProperties))
 		{
@@ -233,7 +233,7 @@ void FObjectPropertyTrace::Init()
 void FObjectPropertyTrace::Destroy()
 {
 	check(ObjectPropertyTrace::TickerHandle.IsValid());
-	FTicker::GetCoreTicker().RemoveTicker(ObjectPropertyTrace::TickerHandle);
+	FTSTicker::GetCoreTicker().RemoveTicker(ObjectPropertyTrace::TickerHandle);
 }
 
 bool FObjectPropertyTrace::IsEnabled()
diff --git a/Engine/Plugins/Animation/GameplayInsights/Source/RewindDebugger/Private/RewindDebugger.cpp b/Engine/Plugins/Animation/GameplayInsights/Source/RewindDebugger/Private/RewindDebugger.cpp
--- a/Engine/Plugins/Animation/GameplayInsights/Source/RewindDebugger/Private/RewindDebugger.cpp
+++ b/Engine/Plugins/Animation/GameplayInsights/Source/RewindDebugger/Private/RewindDebugger.cpp
@@ -59,7 +59,7 @@ FRewindDebugger::FRewindDebugger()  :
 
 	UnrealInsightsModule = &FModuleManager::LoadModuleChecked<IUnrealInsightsModule>("TraceInsights");
 
-	TickerHandle = FTicker::GetCoreTicker().AddTicker(TEXT("RewindDebugger"), 0.0f, [this](float DeltaTime)
+	TickerHandle = FTSTicker::GetCoreTicker().AddTicker(TEXT("RewindDebugger"), 0.0f, [this](float DeltaTime)
 	{
 		QUICK_SCOPE_CYCLE_COUNTER(STAT_FRewindDebuggerModule_Tick);
 
@@ -77,7 +77,7 @@ FRewindDebugger::~FRewindDebugger()
 	FEditorDelegates::EndPIE.RemoveAll(this);
 	FEditorDelegates::SingleStepPIE.RemoveAll(this);
 
-	FTicker::GetCoreTicker().RemoveTicker(TickerHandle);
+	FTSTicker::GetCoreTicker().RemoveTicker(TickerHandle);
 
 }
 
diff --git a/Engine/Plugins/Animation/GameplayInsights/Source/RewindDebugger/Private/RewindDebugger.h b/Engine/Plugins/Animation/GameplayInsights/Source/RewindDebugger/Private/RewindDebugger.h
--- a/Engine/Plugins/Animation/GameplayInsights/Source/RewindDebugger/Private/RewindDebugger.h
+++ b/Engine/Plugins/Animation/GameplayInsights/Source/RewindDebugger/Private/RewindDebugger.h
@@ -148,7 +148,7 @@ private:
 	TMap<uint64, FMeshComponentResetData> MeshComponentsToReset;
 
 	mutable class IUnrealInsightsModule *UnrealInsightsModule;
-	FDelegateHandle TickerHandle;
+	FTSTicker::FDelegateHandle TickerHandle;
 
 	bool bTargetActorPositionValid;
 	FVector TargetActorPosition;
diff --git a/Engine/Plugins/Developer/Concert/ConcertMain/Source/Concert/Private/ConcertClient.cpp b/Engine/Plugins/Developer/Concert/ConcertMain/Source/Concert/Private/ConcertClient.cpp
--- a/Engine/Plugins/Developer/Concert/ConcertMain/Source/Concert/Private/ConcertClient.cpp
+++ b/Engine/Plugins/Developer/Concert/ConcertMain/Source/Concert/Private/ConcertClient.cpp
@@ -52,7 +52,7 @@ public:
 
 		AutoConnectionNotification = MakeAutoConnectNotification();
 
-		AutoConnectionTickHandle = FTicker::GetCoreTicker().AddTicker(TEXT("ConcertAutoConnect"), 1, [this](float)
+		AutoConnectionTickHandle = FTSTicker::GetCoreTicker().AddTicker(TEXT("ConcertAutoConnect"), 1, [this](float)
 		{
 			QUICK_SCOPE_CYCLE_COUNTER(STAT_FConcertAutoConnection_Tick);
 			Tick();
@@ -68,7 +68,7 @@ public:
 
 		if (AutoConnectionTickHandle.IsValid())
 		{
-			FTicker::GetCoreTicker().RemoveTicker(AutoConnectionTickHandle);
+			FTSTicker::GetCoreTicker().RemoveTicker(AutoConnectionTickHandle);
 			AutoConnectionTickHandle.Reset();
 		}
 
@@ -278,7 +278,7 @@ private:
 	}
 
 	TFuture<TSharedFuture<EConcertResponseCode>> OngoingConnectionRequest;
-	FDelegateHandle AutoConnectionTickHandle;
+	FTSTicker::FDelegateHandle AutoConnectionTickHandle;
 	FConcertClient* Client;
 	TWeakPtr<IConcertClientSession> CurrentSession;
 	const UConcertClientConfig* Settings;
@@ -309,7 +309,7 @@ public:
 	{
 		if (ConnectionTick.IsValid())
 		{
-			FTicker::GetCoreTicker().RemoveTicker(ConnectionTick);
+			FTSTicker::GetCoreTicker().RemoveTicker(ConnectionTick);
 		}
 
 		if (ConnectionTasks.Num() > 0)
@@ -348,7 +348,7 @@ public:
 
 		ConnectionTasks[0]->Execute();
 
-		ConnectionTick = FTicker::GetCoreTicker().AddTicker(TEXT("ConcertPendingConnection"), 0.1f, [this](float)
+		ConnectionTick = FTSTicker::GetCoreTicker().AddTicker(TEXT("ConcertPendingConnection"), 0.1f, [this](float)
 		{
 			QUICK_SCOPE_CYCLE_COUNTER(STAT_FConcertPendingConnection_Tick);
 			Tick();
@@ -483,7 +483,7 @@ private:
 
 	FConcertClient* Client;
 	FConfig Config;
-	FDelegateHandle ConnectionTick;
+	FTSTicker::FDelegateHandle ConnectionTick;
 	TPromise<EConcertResponseCode> ConnectionResult;
 	TUniquePtr<FAsyncTaskNotification> Notification;
 	TArray<TUniquePtr<IConcertClientConnectionTask>> ConnectionTasks;
@@ -754,7 +754,7 @@ void FConcertClient::StartDiscovery()
 	{
 		ClientAdminEndpoint->RegisterEventHandler<FConcertAdmin_ServerDiscoveredEvent>(this, &FConcertClient::HandleServerDiscoveryEvent);
 
-		DiscoveryTick = FTicker::GetCoreTicker().AddTicker(TEXT("Discovery"), 1, [this](float DeltaSeconds) {
+		DiscoveryTick = FTSTicker::GetCoreTicker().AddTicker(TEXT("Discovery"), 1, [this](float DeltaSeconds) {
 			QUICK_SCOPE_CYCLE_COUNTER(STAT_FConcertClient_Discovery_Tick);
 			const FDateTime UtcNow = FDateTime::UtcNow();
 			SendDiscoverServersEvent();
@@ -779,7 +779,7 @@ void FConcertClient::StopDiscovery()
 	}
 	if (DiscoveryTick.IsValid())
 	{
-		FTicker::GetCoreTicker().RemoveTicker(DiscoveryTick);
+		FTSTicker::GetCoreTicker().RemoveTicker(DiscoveryTick);
 		DiscoveryTick.Reset();
 	}
 }
diff --git a/Engine/Plugins/Developer/Concert/ConcertMain/Source/Concert/Private/ConcertClient.h b/Engine/Plugins/Developer/Concert/ConcertMain/Source/Concert/Private/ConcertClient.h
--- a/Engine/Plugins/Developer/Concert/ConcertMain/Source/Concert/Private/ConcertClient.h
+++ b/Engine/Plugins/Developer/Concert/ConcertMain/Source/Concert/Private/ConcertClient.h
@@ -8,6 +8,7 @@
 #include "ConcertSettings.h"
 
 #include "UObject/StrongObjectPtr.h"
+#include "Containers/Ticker.h"
 
 class FConcertClientSession;
 class FConcertAutoConnection;
@@ -150,7 +151,7 @@ private:
 	uint32 DiscoveryCount;
 
 	/** Ticker for Discovering Concert Servers */
-	FDelegateHandle DiscoveryTick;
+	FTSTicker::FDelegateHandle DiscoveryTick;
 
 	struct FKnownServer
 	{
diff --git a/Engine/Plugins/Developer/Concert/ConcertMain/Source/Concert/Private/ConcertClientSession.cpp b/Engine/Plugins/Developer/Concert/ConcertMain/Source/Concert/Private/ConcertClientSession.cpp
--- a/Engine/Plugins/Developer/Concert/ConcertMain/Source/Concert/Private/ConcertClientSession.cpp
+++ b/Engine/Plugins/Developer/Concert/ConcertMain/Source/Concert/Private/ConcertClientSession.cpp
@@ -51,7 +51,7 @@ void FConcertClientSession::Startup()
 		ClientSessionEndpoint->RegisterRequestHandler<FConcertSession_CustomRequest, FConcertSession_CustomResponse>(this, &FConcertClientSession::CommonHandleCustomRequest);
 
 		// Setup the session tick
-		SessionTick = FTicker::GetCoreTicker().AddTicker(TEXT("ClientSession"), 0, [this](float DeltaSeconds)
+		SessionTick = FTSTicker::GetCoreTicker().AddTicker(TEXT("ClientSession"), 0, [this](float DeltaSeconds)
 		{
 			QUICK_SCOPE_CYCLE_COUNTER(STAT_FConcertClientSession_Tick);
 			const FDateTime UtcNow = FDateTime::UtcNow();
@@ -81,7 +81,7 @@ void FConcertClientSession::Shutdown()
 		ClientSessionEndpoint->UnregisterRequestHandler<FConcertSession_CustomRequest>();
 
 		// Unregister the session tick
-		FTicker::GetCoreTicker().RemoveTicker(SessionTick);
+		FTSTicker::GetCoreTicker().RemoveTicker(SessionTick);
 		SessionTick.Reset();
 
 		CommonShutdown();
diff --git a/Engine/Plugins/Developer/Concert/ConcertMain/Source/Concert/Private/ConcertClientSession.h b/Engine/Plugins/Developer/Concert/ConcertMain/Source/Concert/Private/ConcertClientSession.h
--- a/Engine/Plugins/Developer/Concert/ConcertMain/Source/Concert/Private/ConcertClientSession.h
+++ b/Engine/Plugins/Developer/Concert/ConcertMain/Source/Concert/Private/ConcertClientSession.h
@@ -5,6 +5,7 @@
 #include "CoreMinimal.h"
 #include "ConcertSession.h"
 #include "ConcertMessages.h"
+#include "Containers/Ticker.h"
 
 class IConcertLocalEndpoint;
 struct FConcertClientSettings;
@@ -171,7 +172,7 @@ private:
 	uint8 SuspendedCount;
 
 	/** Ticker for the session */
-	FDelegateHandle SessionTick;
+	FTSTicker::FDelegateHandle SessionTick;
 
 	/** Last connection tick */
 	FDateTime LastConnectionTick;
diff --git a/Engine/Plugins/Developer/Concert/ConcertMain/Source/Concert/Private/ConcertServerSession.cpp b/Engine/Plugins/Developer/Concert/ConcertMain/Source/Concert/Private/ConcertServerSession.cpp
--- a/Engine/Plugins/Developer/Concert/ConcertMain/Source/Concert/Private/ConcertServerSession.cpp
+++ b/Engine/Plugins/Developer/Concert/ConcertMain/Source/Concert/Private/ConcertServerSession.cpp
@@ -46,7 +46,7 @@ void FConcertServerSession::Startup()
 		ServerSessionEndpoint->RegisterRequestHandler<FConcertSession_CustomRequest, FConcertSession_CustomResponse>(this, &FConcertServerSession::HandleCustomRequest);
 
 		// Setup the session tick
-		SessionTick = FTicker::GetCoreTicker().AddTicker(TEXT("ServerSession"), 0, [this](float DeltaSeconds)
+		SessionTick = FTSTicker::GetCoreTicker().AddTicker(TEXT("ServerSession"), 0, [this](float DeltaSeconds)
 		{
 			QUICK_SCOPE_CYCLE_COUNTER(STAT_FConcertServerSession_Tick);
 			TickConnections(DeltaSeconds);
@@ -74,7 +74,7 @@ void FConcertServerSession::Shutdown()
 		ServerSessionEndpoint->UnregisterRequestHandler<FConcertSession_CustomRequest>();
 
 		// Unregister the session tick
-		FTicker::GetCoreTicker().RemoveTicker(SessionTick);
+		FTSTicker::GetCoreTicker().RemoveTicker(SessionTick);
 		SessionTick.Reset();
 
 		CommonShutdown();
diff --git a/Engine/Plugins/Developer/Concert/ConcertMain/Source/Concert/Private/ConcertServerSession.h b/Engine/Plugins/Developer/Concert/ConcertMain/Source/Concert/Private/ConcertServerSession.h
--- a/Engine/Plugins/Developer/Concert/ConcertMain/Source/Concert/Private/ConcertServerSession.h
+++ b/Engine/Plugins/Developer/Concert/ConcertMain/Source/Concert/Private/ConcertServerSession.h
@@ -5,6 +5,7 @@
 #include "CoreMinimal.h"
 #include "ConcertSession.h"
 #include "ConcertMessages.h"
+#include "Containers/Ticker.h"
 
 class IConcertLocalEndpoint;
 struct FConcertServerSettings;
@@ -136,7 +137,7 @@ private:
 	IConcertLocalEndpointPtr ServerSessionEndpoint;
 
 	/** Ticker for the session */
-	FDelegateHandle SessionTick;
+	FTSTicker::FDelegateHandle SessionTick;
 
 	/** Callback for when a server session ticks */
 	FOnConcertServerSessionTick OnTickDelegate;
diff --git a/Engine/Plugins/Developer/Concert/ConcertMain/Source/ConcertTransport/Private/ConcertLocalEndpoint.cpp b/Engine/Plugins/Developer/Concert/ConcertMain/Source/ConcertTransport/Private/ConcertLocalEndpoint.cpp
--- a/Engine/Plugins/Developer/Concert/ConcertMain/Source/ConcertTransport/Private/ConcertLocalEndpoint.cpp
+++ b/Engine/Plugins/Developer/Concert/ConcertMain/Source/ConcertTransport/Private/ConcertLocalEndpoint.cpp
@@ -96,14 +96,14 @@ FConcertLocalEndpoint::FConcertLocalEndpoint(const FString& InEndpointFriendlyNa
 
 	KeepAliveRunnable = MakeUnique<FConcertLocalEndpointKeepAliveRunnable>(this, *MessageEndpointName);
 
-	TickerHandle = FTicker::GetCoreTicker().AddTicker(FTickerDelegate::CreateRaw(this, &FConcertLocalEndpoint::HandleTick), 0.0f);
+	TickerHandle = FTSTicker::GetCoreTicker().AddTicker(FTickerDelegate::CreateRaw(this, &FConcertLocalEndpoint::HandleTick), 0.0f);
 }
 
 FConcertLocalEndpoint::~FConcertLocalEndpoint()
 {
 	if (TickerHandle.IsValid())
 	{
-		FTicker::GetCoreTicker().RemoveTicker(TickerHandle);
+		FTSTicker::GetCoreTicker().RemoveTicker(TickerHandle);
 	}
 
 	KeepAliveRunnable.Reset();
diff --git a/Engine/Plugins/Developer/Concert/ConcertMain/Source/ConcertTransport/Private/ConcertLocalEndpoint.h b/Engine/Plugins/Developer/Concert/ConcertMain/Source/ConcertTransport/Private/ConcertLocalEndpoint.h
--- a/Engine/Plugins/Developer/Concert/ConcertMain/Source/ConcertTransport/Private/ConcertLocalEndpoint.h
+++ b/Engine/Plugins/Developer/Concert/ConcertMain/Source/ConcertTransport/Private/ConcertLocalEndpoint.h
@@ -8,6 +8,7 @@
 #include "ConcertRemoteEndpoint.h"
 #include "ConcertTransportSettings.h"
 #include "Containers/Queue.h"
+#include "Containers/Ticker.h"
 
 class FMessageEndpoint;
 struct FMessageBusNotification;
@@ -163,7 +164,7 @@ private:
 	TMap<FName, TSharedPtr<IConcertRequestHandler>> RequestHandlers;
 
 	/** Handle to the registered ticker for request and reliable message */
-	FDelegateHandle TickerHandle;
+	FTSTicker::FDelegateHandle TickerHandle;
 
 	/** Runnable thread used to send keep-alive messages even when the game thread is blocked */
 	TUniquePtr<FConcertLocalEndpointKeepAliveRunnable> KeepAliveRunnable;
diff --git a/Engine/Plugins/Developer/Concert/DisasterRecoveryClient/Source/DisasterRecoveryClient/Private/DisasterRecoveryTasks.cpp b/Engine/Plugins/Developer/Concert/DisasterRecoveryClient/Source/DisasterRecoveryClient/Private/DisasterRecoveryTasks.cpp
--- a/Engine/Plugins/Developer/Concert/DisasterRecoveryClient/Source/DisasterRecoveryClient/Private/DisasterRecoveryTasks.cpp
+++ b/Engine/Plugins/Developer/Concert/DisasterRecoveryClient/Source/DisasterRecoveryClient/Private/DisasterRecoveryTasks.cpp
@@ -31,12 +31,12 @@ FText GetSystemErrorMessage(EConcertResponseCode Code)
 
 FDisasterRecoveryTaskExecutor::FDisasterRecoveryTaskExecutor()
 {
-	TickerHandle = FTicker::GetCoreTicker().AddTicker(FTickerDelegate::CreateRaw(this, &FDisasterRecoveryTaskExecutor::Tick), 0.0f);
+	TickerHandle = FTSTicker::GetCoreTicker().AddTicker(FTickerDelegate::CreateRaw(this, &FDisasterRecoveryTaskExecutor::Tick), 0.0f);
 }
 
 FDisasterRecoveryTaskExecutor::~FDisasterRecoveryTaskExecutor()
 {
-	FTicker::GetCoreTicker().RemoveTicker(TickerHandle);
+	FTSTicker::GetCoreTicker().RemoveTicker(TickerHandle);
 	Abort();
 }
 
diff --git a/Engine/Plugins/Editor/AssetSearch/Source/Private/AssetSearchManager.cpp b/Engine/Plugins/Editor/AssetSearch/Source/Private/AssetSearchManager.cpp
--- a/Engine/Plugins/Editor/AssetSearch/Source/Private/AssetSearchManager.cpp
+++ b/Engine/Plugins/Editor/AssetSearch/Source/Private/AssetSearchManager.cpp
@@ -208,7 +208,7 @@ FAssetSearchManager::~FAssetSearchManager()
 	FCoreUObjectDelegates::OnAssetLoaded.RemoveAll(this);
 	UObject::FAssetRegistryTag::OnGetExtraObjectTags.RemoveAll(this);
 
-	FTicker::GetCoreTicker().RemoveTicker(TickerHandle);
+	FTSTicker::GetCoreTicker().RemoveTicker(TickerHandle);
 }
 
 void FAssetSearchManager::Start()
@@ -232,7 +232,7 @@ void FAssetSearchManager::Start()
 	UPackage::PackageSavedWithContextEvent.AddRaw(this, &FAssetSearchManager::HandlePackageSaved);
 	FCoreUObjectDelegates::OnAssetLoaded.AddRaw(this, &FAssetSearchManager::OnAssetLoaded);
 
-	TickerHandle = FTicker::GetCoreTicker().AddTicker(FTickerDelegate::CreateRaw(this, &FAssetSearchManager::Tick_GameThread), 0);
+	TickerHandle = FTSTicker::GetCoreTicker().AddTicker(FTickerDelegate::CreateRaw(this, &FAssetSearchManager::Tick_GameThread), 0);
 
 	RunThread = true;
 	DatabaseThread = FRunnableThread::Create(this, TEXT("UniversalSearch"), 0, TPri_BelowNormal);
diff --git a/Engine/Plugins/Enterprise/VariantManager/Source/VariantManager/Private/SVariantManager.cpp b/Engine/Plugins/Enterprise/VariantManager/Source/VariantManager/Private/SVariantManager.cpp
--- a/Engine/Plugins/Enterprise/VariantManager/Source/VariantManager/Private/SVariantManager.cpp
+++ b/Engine/Plugins/Enterprise/VariantManager/Source/VariantManager/Private/SVariantManager.cpp
@@ -1852,7 +1852,7 @@ void SVariantManager::SwitchOnVariant(UVariant* Variant)
 	// Do this on next frame because there is some minor issue where if visibility changes are triggered
 	// on the same frame that is meant to be invalidated, sometimes the primitive's occlusion history doesn't
 	// refresh properly and we get some incorrectly hidden/visible objects (UE-100896)
-	FTicker::GetCoreTicker().AddTicker(
+	FTSTicker::GetCoreTicker().AddTicker(
 		FTickerDelegate::CreateLambda([](float Time)
 		{
 			const bool bInvalidateHitProxies = false;
diff --git a/Engine/Plugins/Experimental/AutomationUtils/Source/AutomationUtils/Private/AutomationUtilsBlueprintLibrary.cpp b/Engine/Plugins/Experimental/AutomationUtils/Source/AutomationUtils/Private/AutomationUtilsBlueprintLibrary.cpp
--- a/Engine/Plugins/Experimental/AutomationUtils/Source/AutomationUtils/Private/AutomationUtilsBlueprintLibrary.cpp
+++ b/Engine/Plugins/Experimental/AutomationUtils/Source/AutomationUtils/Private/AutomationUtilsBlueprintLibrary.cpp
@@ -231,7 +231,7 @@ FAutomationUtilsGameplayAutomationScreenshotInstance::FAutomationUtilsGameplayAu
 	}
 	else
 	{
-		FTicker::GetCoreTicker().AddTicker(TEXT("FAutomationUtilsGameplayScreenshotInstanceAutoCleanup"), 0.1f, [this](float)
+		FTSTicker::GetCoreTicker().AddTicker(TEXT("FAutomationUtilsGameplayScreenshotInstanceAutoCleanup"), 0.1f, [this](float)
 		{
 			QUICK_SCOPE_CYCLE_COUNTER(STAT_FAutomationUtilsGameplayScreenshotInstanceAutoCleanup);
 			FAutomationUtilsGameplayAutomationScreenshotFactory::RequestDeleteScreenshotInstance(ScreenshotName);
diff --git a/Engine/Plugins/Experimental/CommonUI/Source/CommonInput/Private/CommonInputSubsystem.cpp b/Engine/Plugins/Experimental/CommonUI/Source/CommonInput/Private/CommonInputSubsystem.cpp
--- a/Engine/Plugins/Experimental/CommonUI/Source/CommonInput/Private/CommonInputSubsystem.cpp
+++ b/Engine/Plugins/Experimental/CommonUI/Source/CommonInput/Private/CommonInputSubsystem.cpp
@@ -320,7 +320,7 @@ void UCommonInputSubsystem::Initialize(FSubsystemCollectionBase& Collection)
 	CommonInputPreprocessor = MakeShared<FCommonInputPreprocessor>(*this);
 	FSlateApplication::Get().RegisterInputPreProcessor(CommonInputPreprocessor, 0);
 
-	TickHandle = FTicker::GetCoreTicker().AddTicker(FTickerDelegate::CreateUObject(this, &UCommonInputSubsystem::Tick), 0.1f);
+	TickHandle = FTSTicker::GetCoreTicker().AddTicker(FTickerDelegate::CreateUObject(this, &UCommonInputSubsystem::Tick), 0.1f);
 
 	CVarInputKeysVisible->SetOnChangedCallback(FConsoleVariableDelegate::CreateUObject(this, &UCommonInputSubsystem::ShouldShowInputKeysChanged));
 }
@@ -331,7 +331,7 @@ void UCommonInputSubsystem::Deinitialize()
 	FSlateApplication::Get().UnregisterInputPreProcessor(CommonInputPreprocessor);
 	CommonInputPreprocessor.Reset();
 
-	FTicker::GetCoreTicker().RemoveTicker(TickHandle);
+	FTSTicker::GetCoreTicker().RemoveTicker(TickHandle);
 }
 
 FGamepadChangeDetectedEvent& UCommonInputSubsystem::GetOnGamepadChangeDetected()
diff --git a/Engine/Plugins/Experimental/CommonUI/Source/CommonInput/Public/CommonInputSubsystem.h b/Engine/Plugins/Experimental/CommonUI/Source/CommonInput/Public/CommonInputSubsystem.h
--- a/Engine/Plugins/Experimental/CommonUI/Source/CommonInput/Public/CommonInputSubsystem.h
+++ b/Engine/Plugins/Experimental/CommonUI/Source/CommonInput/Public/CommonInputSubsystem.h
@@ -111,7 +111,7 @@ private:
 
 	bool CheckForInputMethodThrashing(ECommonInputType NewInputType);
 
-	FDelegateHandle TickHandle;
+	FTSTicker::FDelegateHandle TickHandle;
 
 	UPROPERTY(BlueprintAssignable, Category = CommonInputSubsystem, meta = (AllowPrivateAccess))
 	FInputMethodChangedDelegate OnInputMethodChanged;
diff --git a/Engine/Plugins/Experimental/CommonUI/Source/CommonUI/Private/CommonWidgetCarousel.cpp b/Engine/Plugins/Experimental/CommonUI/Source/CommonUI/Private/CommonWidgetCarousel.cpp
--- a/Engine/Plugins/Experimental/CommonUI/Source/CommonUI/Private/CommonWidgetCarousel.cpp
+++ b/Engine/Plugins/Experimental/CommonUI/Source/CommonUI/Private/CommonWidgetCarousel.cpp
@@ -25,14 +25,14 @@ void UCommonWidgetCarousel::ReleaseSlateResources(bool bReleaseChildren)
 void UCommonWidgetCarousel::BeginAutoScrolling(float ScrollInterval)
 {
 	EndAutoScrolling();
-	TickerHandle = FTicker::GetCoreTicker().AddTicker(FTickerDelegate::CreateUObject(this, &UCommonWidgetCarousel::AutoScrollCallback), ScrollInterval);
+	TickerHandle = FTSTicker::GetCoreTicker().AddTicker(FTickerDelegate::CreateUObject(this, &UCommonWidgetCarousel::AutoScrollCallback), ScrollInterval);
 }
 
 void UCommonWidgetCarousel::EndAutoScrolling()
 {
 	if ( TickerHandle.IsValid() )
 	{
-		FTicker::GetCoreTicker().RemoveTicker(TickerHandle);
+		FTSTicker::GetCoreTicker().RemoveTicker(TickerHandle);
 		TickerHandle.Reset();
 	}
 }
diff --git a/Engine/Plugins/Experimental/CommonUI/Source/CommonUI/Private/Input/CommonUIActionRouterBase.cpp b/Engine/Plugins/Experimental/CommonUI/Source/CommonUI/Private/Input/CommonUIActionRouterBase.cpp
--- a/Engine/Plugins/Experimental/CommonUI/Source/CommonUI/Private/Input/CommonUIActionRouterBase.cpp
+++ b/Engine/Plugins/Experimental/CommonUI/Source/CommonUI/Private/Input/CommonUIActionRouterBase.cpp
@@ -201,8 +201,8 @@ void UCommonUIActionRouterBase::RegisterAnalogCursorTick()
 	FSlateApplication::Get().RegisterInputPreProcessor(AnalogCursor, UCommonUIInputSettings::Get().GetAnalogCursorSettings().PreprocessorPriority);
 	if (bIsActivatableTreeEnabled)
 	{
-		FTicker::GetCoreTicker().RemoveTicker(TickHandle);
-		TickHandle = FTicker::GetCoreTicker().AddTicker(FTickerDelegate::CreateUObject(this, &UCommonUIActionRouterBase::Tick));
+		FTSTicker::GetCoreTicker().RemoveTicker(TickHandle);
+		TickHandle = FTSTicker::GetCoreTicker().AddTicker(FTickerDelegate::CreateUObject(this, &UCommonUIActionRouterBase::Tick));
 	}
 }
 
@@ -212,7 +212,7 @@ void UCommonUIActionRouterBase::Deinitialize()
 	
 	FSlateApplication::Get().OnFocusChanging().RemoveAll(this);
 	FSlateApplication::Get().UnregisterInputPreProcessor(AnalogCursor);
-	FTicker::GetCoreTicker().RemoveTicker(TickHandle);
+	FTSTicker::GetCoreTicker().RemoveTicker(TickHandle);
 	SetActiveRoot(nullptr);
 	HeldKeys.Empty();
 }
diff --git a/Engine/Plugins/Experimental/CommonUI/Source/CommonUI/Private/Widgets/CommonActivatableWidgetContainer.cpp b/Engine/Plugins/Experimental/CommonUI/Source/CommonUI/Private/Widgets/CommonActivatableWidgetContainer.cpp
--- a/Engine/Plugins/Experimental/CommonUI/Source/CommonUI/Private/Widgets/CommonActivatableWidgetContainer.cpp
+++ b/Engine/Plugins/Experimental/CommonUI/Source/CommonUI/Private/Widgets/CommonActivatableWidgetContainer.cpp
@@ -218,7 +218,7 @@ void UCommonActivatableWidgetContainerBase::ReleaseWidget(const TSharedRef<SWidg
 		ReleasedWidgets.Add(WidgetToRelease);
 		if (ReleasedWidgets.Num() == 1)
 		{
-			FTicker::GetCoreTicker().AddTicker(FTickerDelegate::CreateWeakLambda(this,
+			FTSTicker::GetCoreTicker().AddTicker(FTickerDelegate::CreateWeakLambda(this,
 				[this](float)
 				{
 					QUICK_SCOPE_CYCLE_COUNTER(STAT_UCommonActivatableWidgetContainerBase_ReleaseWidget);
diff --git a/Engine/Plugins/Experimental/CommonUI/Source/CommonUI/Public/Input/CommonUIActionRouterBase.h b/Engine/Plugins/Experimental/CommonUI/Source/CommonUI/Public/Input/CommonUIActionRouterBase.h
--- a/Engine/Plugins/Experimental/CommonUI/Source/CommonUI/Public/Input/CommonUIActionRouterBase.h
+++ b/Engine/Plugins/Experimental/CommonUI/Source/CommonUI/Public/Input/CommonUIActionRouterBase.h
@@ -6,6 +6,7 @@
 #include "CommonUIInputTypes.h"
 #include "CommonInputBaseTypes.h"
 #include "Misc/Passkey.h"
+#include "Containers/Ticker.h"
 
 #include "CommonUIActionRouterBase.generated.h"
 
@@ -123,7 +124,7 @@ protected:
 	bool bIsActivatableTreeEnabled = true;
 
 	TSharedPtr<FCommonAnalogCursor> AnalogCursor;
-	FDelegateHandle TickHandle;
+	FTSTicker::FDelegateHandle TickHandle;
 
 private:
 	bool Tick(float DeltaTime);
diff --git a/Engine/Plugins/Experimental/GameFeatures/Source/GameFeatures/Private/GameFeaturePluginStateMachine.cpp b/Engine/Plugins/Experimental/GameFeatures/Source/GameFeatures/Private/GameFeaturePluginStateMachine.cpp
--- a/Engine/Plugins/Experimental/GameFeatures/Source/GameFeatures/Private/GameFeaturePluginStateMachine.cpp
+++ b/Engine/Plugins/Experimental/GameFeatures/Source/GameFeatures/Private/GameFeaturePluginStateMachine.cpp
@@ -96,7 +96,7 @@ FGameFeaturePluginState::~FGameFeaturePluginState()
 {
 	if (TickHandle.IsValid())
 	{
-		FTicker::GetCoreTicker().RemoveTicker(TickHandle);
+		FTSTicker::GetCoreTicker().RemoveTicker(TickHandle);
 		TickHandle.Reset();
 	}
 }
@@ -105,10 +105,10 @@ void FGameFeaturePluginState::UpdateStateMachineDeferred(float Delay /*= 0.0f*/)
 {
 	if (TickHandle.IsValid())
 	{
-		FTicker::GetCoreTicker().RemoveTicker(TickHandle);
+		FTSTicker::GetCoreTicker().RemoveTicker(TickHandle);
 	}
 
-	TickHandle = FTicker::GetCoreTicker().AddTicker(FTickerDelegate::CreateLambda([this](float dts) mutable
+	TickHandle = FTSTicker::GetCoreTicker().AddTicker(FTickerDelegate::CreateLambda([this](float dts) mutable
 	{
 		StateProperties.OnRequestUpdateStateMachine.ExecuteIfBound();
 		TickHandle.Reset();
@@ -471,14 +471,14 @@ struct FGameFeaturePluginState_Downloading : public FGameFeaturePluginState
 	bool bPluginDownloaded = false;
 	TArray<FName> PendingBundleDownloads;
 	TUniquePtr<FInstallBundleCombinedProgressTracker> ProgressTracker;
-	FDelegateHandle ProgressUpdateHandle;
+	FTSTicker::FDelegateHandle ProgressUpdateHandle;
 	FDelegateHandle GotContentStateHandle;
 
 	void Cleanup()
 	{
 		if (ProgressUpdateHandle.IsValid())
 		{
-			FTicker::GetCoreTicker().RemoveTicker(ProgressUpdateHandle);
+			FTSTicker::GetCoreTicker().RemoveTicker(ProgressUpdateHandle);
 			ProgressUpdateHandle.Reset();
 		}
 
@@ -543,7 +543,7 @@ struct FGameFeaturePluginState_Downloading : public FGameFeaturePluginState
 			ProgressTracker = MakeUnique<FInstallBundleCombinedProgressTracker>(false);
 			ProgressTracker->SetBundlesToTrackFromContentState(BundleContentState, PendingBundleDownloads);
 
-			ProgressUpdateHandle = FTicker::GetCoreTicker().AddTicker(
+			ProgressUpdateHandle = FTSTicker::GetCoreTicker().AddTicker(
 				FTickerDelegate::CreateRaw(this, &FGameFeaturePluginState_Downloading::OnUpdateProgress)/*, 0.1f*/);
 		}
 	}
diff --git a/Engine/Plugins/Experimental/GameFeatures/Source/GameFeatures/Private/GameFeaturePluginStateMachine.h b/Engine/Plugins/Experimental/GameFeatures/Source/GameFeatures/Private/GameFeaturePluginStateMachine.h
--- a/Engine/Plugins/Experimental/GameFeatures/Source/GameFeatures/Private/GameFeaturePluginStateMachine.h
+++ b/Engine/Plugins/Experimental/GameFeatures/Source/GameFeatures/Private/GameFeaturePluginStateMachine.h
@@ -5,6 +5,7 @@
 #include "CoreMinimal.h"
 #include "Misc/EnumRange.h"
 #include "Containers/Union.h"
+#include "Containers/Ticker.h"
 #include "UObject/ObjectMacros.h"
 #include "UObject/Object.h"
 #include "GameFeaturePluginOperationResult.h"
@@ -318,7 +319,7 @@ struct FGameFeaturePluginState
 	void UpdateProgress(float Progress) const;
 
 private:
-	mutable FDelegateHandle TickHandle;
+	mutable FTSTicker::FDelegateHandle TickHandle;
 };
 
 /** Information about a given plugin state, used to expose information to external code */
diff --git a/Engine/Plugins/Experimental/GameFeatures/Source/GameFeatures/Private/GameFeaturesSubsystem.cpp b/Engine/Plugins/Experimental/GameFeatures/Source/GameFeatures/Private/GameFeaturesSubsystem.cpp
--- a/Engine/Plugins/Experimental/GameFeatures/Source/GameFeatures/Private/GameFeaturesSubsystem.cpp
+++ b/Engine/Plugins/Experimental/GameFeatures/Source/GameFeatures/Private/GameFeaturesSubsystem.cpp
@@ -454,7 +454,7 @@ void UGameFeaturesSubsystem::LoadGameFeaturePlugin(const FString& PluginURL, con
 	{
 		if (StateMachine->GetCurrentState() >= EGameFeaturePluginState::Loaded)
 		{
-			FTicker::GetCoreTicker().AddTicker(FTickerDelegate::CreateWeakLambda(this, [this, CompleteDelegate](float dts)
+			FTSTicker::GetCoreTicker().AddTicker(FTickerDelegate::CreateWeakLambda(this, [this, CompleteDelegate](float dts)
 				{
 					CompleteDelegate.ExecuteIfBound(MakeValue());
 					return false;
@@ -509,7 +509,7 @@ void UGameFeaturesSubsystem::ChangeGameFeatureTargetState(const FString& PluginU
 					CastChecked<IGameFeatureStateChangeObserver>(Observer)->OnGameFeatureActivating(StateMachine->GetGameFeatureDataForActivePlugin());
 	}
 
-				FTicker::GetCoreTicker().AddTicker(FTickerDelegate::CreateWeakLambda(this, [CompleteDelegate](float)
+				FTSTicker::GetCoreTicker().AddTicker(FTickerDelegate::CreateWeakLambda(this, [CompleteDelegate](float)
 					{
 						CompleteDelegate.ExecuteIfBound(UE::GameFeatures::FResult(MakeValue()));
 						return false;
@@ -518,7 +518,7 @@ void UGameFeaturesSubsystem::ChangeGameFeatureTargetState(const FString& PluginU
 			else if (TargetPluginState > StateMachine->GetCurrentState()
 				&& !GameSpecificPolicies->IsPluginAllowed(PluginURL))
 			{
-				FTicker::GetCoreTicker().AddTicker(FTickerDelegate::CreateWeakLambda(this, [CompleteDelegate](float)
+				FTSTicker::GetCoreTicker().AddTicker(FTickerDelegate::CreateWeakLambda(this, [CompleteDelegate](float)
 					{
 						CompleteDelegate.ExecuteIfBound(MakeError(TEXT("GameFeaturePlugin.StateMachine.Plugin_Denied_By_GameSpecificPolicy")));
 						return false;
@@ -585,7 +585,7 @@ void UGameFeaturesSubsystem::DeactivateGameFeaturePlugin(const FString& PluginUR
 	{
 		if (StateMachine->GetDestinationState() <= EGameFeaturePluginState::Loaded)
 		{
-			FTicker::GetCoreTicker().AddTicker(FTickerDelegate::CreateWeakLambda(this, [CompleteDelegate](float)
+			FTSTicker::GetCoreTicker().AddTicker(FTickerDelegate::CreateWeakLambda(this, [CompleteDelegate](float)
 		{
 					CompleteDelegate.ExecuteIfBound(UE::GameFeatures::FResult(MakeValue()));
 					return false;
@@ -598,7 +598,7 @@ void UGameFeaturesSubsystem::DeactivateGameFeaturePlugin(const FString& PluginUR
 	}
 	else
 	{
-		FTicker::GetCoreTicker().AddTicker(FTickerDelegate::CreateWeakLambda(this, [CompleteDelegate](float)
+		FTSTicker::GetCoreTicker().AddTicker(FTickerDelegate::CreateWeakLambda(this, [CompleteDelegate](float)
 		{
 			CompleteDelegate.ExecuteIfBound(UE::GameFeatures::FResult(MakeError(TEXT("GameFeaturePlugin.BadURL"))));
 			return false;
@@ -619,7 +619,7 @@ void UGameFeaturesSubsystem::UnloadGameFeaturePlugin(const FString& PluginURL, c
 	{
 		if (StateMachine->GetDestinationState() <= EGameFeaturePluginState::Loaded)
 		{
-			FTicker::GetCoreTicker().AddTicker(FTickerDelegate::CreateWeakLambda(this, [CompleteDelegate](float)
+			FTSTicker::GetCoreTicker().AddTicker(FTickerDelegate::CreateWeakLambda(this, [CompleteDelegate](float)
 				{
 					CompleteDelegate.ExecuteIfBound(UE::GameFeatures::FResult(MakeValue()));
 					return false;
@@ -634,7 +634,7 @@ void UGameFeaturesSubsystem::UnloadGameFeaturePlugin(const FString& PluginURL, c
 	}
 	else
 	{
-		FTicker::GetCoreTicker().AddTicker(FTickerDelegate::CreateWeakLambda(this, [CompleteDelegate](float)
+		FTSTicker::GetCoreTicker().AddTicker(FTickerDelegate::CreateWeakLambda(this, [CompleteDelegate](float)
 		{
 			CompleteDelegate.ExecuteIfBound(UE::GameFeatures::FResult(MakeError(TEXT("GameFeaturePlugin.BadURL"))));
 			return false;
@@ -666,7 +666,7 @@ void UGameFeaturesSubsystem::UninstallGameFeaturePlugin(const FString& PluginURL
 		}
 		else if (StateMachine->GetCurrentState() <= EGameFeaturePluginState::StatusKnown)
 		{
-			FTicker::GetCoreTicker().AddTicker(FTickerDelegate::CreateWeakLambda(this, [this, CompleteDelegate](float dts)
+			FTSTicker::GetCoreTicker().AddTicker(FTickerDelegate::CreateWeakLambda(this, [this, CompleteDelegate](float dts)
 			{
 				CompleteDelegate.ExecuteIfBound(UE::GameFeatures::FResult(MakeValue()));
 				return false;
@@ -675,7 +675,7 @@ void UGameFeaturesSubsystem::UninstallGameFeaturePlugin(const FString& PluginURL
 	}
 	else
 	{
-		FTicker::GetCoreTicker().AddTicker(FTickerDelegate::CreateWeakLambda(this, [this, CompleteDelegate](float dts)
+		FTSTicker::GetCoreTicker().AddTicker(FTickerDelegate::CreateWeakLambda(this, [this, CompleteDelegate](float dts)
 		{
 			CompleteDelegate.ExecuteIfBound(UE::GameFeatures::FResult(MakeError(TEXT("GameFeaturePlugin.BadURL"))));
 			return false;
@@ -701,7 +701,7 @@ void UGameFeaturesSubsystem::TerminateGameFeaturePlugin(const FString& PluginURL
 		}
 		else if (StateMachine->GetCurrentState() <= EGameFeaturePluginState::Terminal)
 		{
-			FTicker::GetCoreTicker().AddTicker(FTickerDelegate::CreateWeakLambda(this, [this, CompleteDelegate](float dts)
+			FTSTicker::GetCoreTicker().AddTicker(FTickerDelegate::CreateWeakLambda(this, [this, CompleteDelegate](float dts)
 			{
 				CompleteDelegate.ExecuteIfBound(UE::GameFeatures::FResult(MakeValue()));
 				return false;
@@ -710,7 +710,7 @@ void UGameFeaturesSubsystem::TerminateGameFeaturePlugin(const FString& PluginURL
 	}
 	else
 	{
-		FTicker::GetCoreTicker().AddTicker(FTickerDelegate::CreateWeakLambda(this, [this, CompleteDelegate](float dts)
+		FTSTicker::GetCoreTicker().AddTicker(FTickerDelegate::CreateWeakLambda(this, [this, CompleteDelegate](float dts)
 		{
 			CompleteDelegate.ExecuteIfBound(UE::GameFeatures::FResult(MakeError(TEXT("GameFeaturePlugin.BadURL"))));
 			return false;
diff --git a/Engine/Plugins/Experimental/Gauntlet/Source/Gauntlet/Private/GauntletModule.cpp b/Engine/Plugins/Experimental/Gauntlet/Source/Gauntlet/Private/GauntletModule.cpp
--- a/Engine/Plugins/Experimental/Gauntlet/Source/Gauntlet/Private/GauntletModule.cpp
+++ b/Engine/Plugins/Experimental/Gauntlet/Source/Gauntlet/Private/GauntletModule.cpp
@@ -73,7 +73,7 @@ protected:
 private:
 
 	/** Handle to our tick callback */
-	FDelegateHandle					TickHandle;
+	FTSTicker::FDelegateHandle					TickHandle;
 
 	/** Currently active controllers */
 	TArray<UGauntletTestController*>	Controllers;
@@ -146,7 +146,7 @@ void FGauntletModuleImpl::OnPostEngineInit()
 	FParse::Value(FCommandLine::Get(), TEXT("gauntlet.tickrate="), kTickRate);
 
 
-	TickHandle = FTicker::GetCoreTicker().AddTicker(FTickerDelegate::CreateLambda([this, kTickRate](float TimeDelta)
+	TickHandle = FTSTicker::GetCoreTicker().AddTicker(FTickerDelegate::CreateLambda([this, kTickRate](float TimeDelta)
 	{
 		// ticker passes in frame-delta, not tick delta...
 		InnerTick(kTickRate);
@@ -163,7 +163,7 @@ void FGauntletModuleImpl::ShutdownModule()
 
 	if (TickHandle.IsValid())
 	{
-		FTicker::GetCoreTicker().RemoveTicker(TickHandle);
+		FTSTicker::GetCoreTicker().RemoveTicker(TickHandle);
 		TickHandle.Reset();
 	}
 
diff --git a/Engine/Plugins/Experimental/GeometryCache/Source/GeometryCacheStreamer/Private/GeometryCacheStreamer.cpp b/Engine/Plugins/Experimental/GeometryCache/Source/GeometryCacheStreamer/Private/GeometryCacheStreamer.cpp
--- a/Engine/Plugins/Experimental/GeometryCache/Source/GeometryCacheStreamer/Private/GeometryCacheStreamer.cpp
+++ b/Engine/Plugins/Experimental/GeometryCache/Source/GeometryCacheStreamer/Private/GeometryCacheStreamer.cpp
@@ -28,7 +28,7 @@ public:
 	//~ End IGeometryCacheStreamer Interface
 
 private:
-	FDelegateHandle TickHandle;
+	FTSTicker::FDelegateHandle TickHandle;
 
 	typedef TMap<UGeometryCacheTrack*, IGeometryCacheStream*> FTracksToStreams;
 	FTracksToStreams TracksToStreams;
@@ -46,7 +46,7 @@ FGeometryCacheStreamer::FGeometryCacheStreamer()
 , CurrentIndex(0)
 
 {
-	TickHandle = FTicker::GetCoreTicker().AddTicker(
+	TickHandle = FTSTicker::GetCoreTicker().AddTicker(
 		FTickerDelegate::CreateLambda([this](float Time)
 		{
 			this->Tick(Time);
@@ -59,7 +59,7 @@ FGeometryCacheStreamer::~FGeometryCacheStreamer()
 {
 	if (TickHandle.IsValid())
 	{
-		FTicker::GetCoreTicker().RemoveTicker(TickHandle);
+		FTSTicker::GetCoreTicker().RemoveTicker(TickHandle);
 		TickHandle.Reset();
 	}
 }
diff --git a/Engine/Plugins/Experimental/PythonScriptPlugin/Source/PythonScriptPlugin/Private/PythonScriptCommandlet.cpp b/Engine/Plugins/Experimental/PythonScriptPlugin/Source/PythonScriptPlugin/Private/PythonScriptCommandlet.cpp
--- a/Engine/Plugins/Experimental/PythonScriptPlugin/Source/PythonScriptPlugin/Private/PythonScriptCommandlet.cpp
+++ b/Engine/Plugins/Experimental/PythonScriptPlugin/Source/PythonScriptPlugin/Private/PythonScriptCommandlet.cpp
@@ -36,7 +36,7 @@ int32 UPythonScriptCommandlet::Main(const FString& Params)
 #if WITH_PYTHON
 	{
 		// Tick once to ensure that any start-up scripts have been run
-		FTicker::GetCoreTicker().Tick(0.0f);
+		FTSTicker::GetCoreTicker().Tick(0.0f);
 
 		UE_LOG(LogPythonScriptCommandlet, Display, TEXT("Running Python script: %s"), *PythonScript);
 
diff --git a/Engine/Plugins/Experimental/PythonScriptPlugin/Source/PythonScriptPlugin/Private/PythonScriptPlugin.cpp b/Engine/Plugins/Experimental/PythonScriptPlugin/Source/PythonScriptPlugin/Private/PythonScriptPlugin.cpp
--- a/Engine/Plugins/Experimental/PythonScriptPlugin/Source/PythonScriptPlugin/Private/PythonScriptPlugin.cpp
+++ b/Engine/Plugins/Experimental/PythonScriptPlugin/Source/PythonScriptPlugin/Private/PythonScriptPlugin.cpp
@@ -816,7 +816,7 @@ void FPythonScriptPlugin::InitializePython()
 		FPyWrapperTypeRegistry::Get().GenerateWrappedTypes();
 
 		// Initialize the tick handler
-		TickHandle = FTicker::GetCoreTicker().AddTicker(FTickerDelegate::CreateLambda([this](float DeltaTime)
+		TickHandle = FTSTicker::GetCoreTicker().AddTicker(FTickerDelegate::CreateLambda([this](float DeltaTime)
 		{
 			QUICK_SCOPE_CYCLE_COUNTER(STAT_FPythonScriptPlugin_Tick);
 			Tick(DeltaTime);
@@ -954,10 +954,10 @@ void FPythonScriptPlugin::ShutdownPython()
 	// Notify any external listeners
 	OnPythonShutdownDelegate.Broadcast();
 
-	FTicker::GetCoreTicker().RemoveTicker(TickHandle);
+	FTSTicker::GetCoreTicker().RemoveTicker(TickHandle);
 	if (ModuleDelayedHandle.IsValid())
 	{
-		FTicker::GetCoreTicker().RemoveTicker(ModuleDelayedHandle);
+		FTSTicker::GetCoreTicker().RemoveTicker(ModuleDelayedHandle);
 	}
 
 	FPyWrapperTypeRegistry::Get().OnModuleDirtied().RemoveAll(this);
@@ -1014,12 +1014,12 @@ void FPythonScriptPlugin::RequestStubCodeGeneration()
 	// If there is an existing pending notification, remove it so that it can be reset
 	if (ModuleDelayedHandle.IsValid())
 	{
-		FTicker::GetCoreTicker().RemoveTicker(ModuleDelayedHandle);
+		FTSTicker::GetCoreTicker().RemoveTicker(ModuleDelayedHandle);
 		ModuleDelayedHandle.Reset();
 	}
 
 	// Set new tick
-	ModuleDelayedHandle = FTicker::GetCoreTicker().AddTicker(FTickerDelegate::CreateLambda(
+	ModuleDelayedHandle = FTSTicker::GetCoreTicker().AddTicker(FTickerDelegate::CreateLambda(
 		[this](float DeltaTime)
 		{
 			QUICK_SCOPE_CYCLE_COUNTER(STAT_FPythonScriptPlugin_ModuleDelayed);
diff --git a/Engine/Plugins/Experimental/PythonScriptPlugin/Source/PythonScriptPlugin/Private/PythonScriptPlugin.h b/Engine/Plugins/Experimental/PythonScriptPlugin/Source/PythonScriptPlugin/Private/PythonScriptPlugin.h
--- a/Engine/Plugins/Experimental/PythonScriptPlugin/Source/PythonScriptPlugin/Private/PythonScriptPlugin.h
+++ b/Engine/Plugins/Experimental/PythonScriptPlugin/Source/PythonScriptPlugin/Private/PythonScriptPlugin.h
@@ -195,8 +195,8 @@ private:
 	FPythonCommandExecutor CmdExec;
 	FPythonREPLCommandExecutor CmdREPLExec;
 	IPythonCommandMenu* CmdMenu;
-	FDelegateHandle TickHandle;
-	FDelegateHandle ModuleDelayedHandle;
+	FTSTicker::FDelegateHandle TickHandle;
+	FTSTicker::FDelegateHandle ModuleDelayedHandle;
 
 	PyUtil::FPyApiBuffer PyProgramName;
 	PyUtil::FPyApiBuffer PyHomePath;
diff --git a/Engine/Plugins/Experimental/Text3D/Source/Text3D/Private/Text3DEngineSubsystem.cpp b/Engine/Plugins/Experimental/Text3D/Source/Text3D/Private/Text3DEngineSubsystem.cpp
--- a/Engine/Plugins/Experimental/Text3D/Source/Text3D/Private/Text3DEngineSubsystem.cpp
+++ b/Engine/Plugins/Experimental/Text3D/Source/Text3D/Private/Text3DEngineSubsystem.cpp
@@ -49,14 +49,14 @@ UText3DEngineSubsystem::UText3DEngineSubsystem()
 void UText3DEngineSubsystem::Initialize(FSubsystemCollectionBase& Collection)
 {
 	Super::Initialize(Collection);
-	CleanupTickerHandle = FTicker::GetCoreTicker().AddTicker(FTickerDelegate::CreateUObject(this, &UText3DEngineSubsystem::CleanupTimerCallback), 600.0f);
+	CleanupTickerHandle = FTSTicker::GetCoreTicker().AddTicker(FTickerDelegate::CreateUObject(this, &UText3DEngineSubsystem::CleanupTimerCallback), 600.0f);
 }
 
 void UText3DEngineSubsystem::Deinitialize()
 {
 	if (CleanupTickerHandle.IsValid())
 	{
-		FTicker::GetCoreTicker().RemoveTicker(CleanupTickerHandle);
+		FTSTicker::GetCoreTicker().RemoveTicker(CleanupTickerHandle);
 		CleanupTickerHandle.Reset();
 	}
 
diff --git a/Engine/Plugins/Experimental/Text3D/Source/Text3D/Private/Text3DEngineSubsystem.h b/Engine/Plugins/Experimental/Text3D/Source/Text3D/Private/Text3DEngineSubsystem.h
--- a/Engine/Plugins/Experimental/Text3D/Source/Text3D/Private/Text3DEngineSubsystem.h
+++ b/Engine/Plugins/Experimental/Text3D/Source/Text3D/Private/Text3DEngineSubsystem.h
@@ -6,6 +6,8 @@
 #include "Text3DPrivate.h"
 #include "BevelType.h"
 #include "ContourNode.h"
+#include "Containers/Ticker.h"
+
 #include "Text3DEngineSubsystem.generated.h"
 
 
@@ -92,5 +94,5 @@ private:
 	UPROPERTY()
 	TMap<uint32, FCachedFontData> CachedFonts;
 
-	FDelegateHandle CleanupTickerHandle;
+	FTSTicker::FDelegateHandle CleanupTickerHandle;
 };
diff --git a/Engine/Plugins/FX/Niagara/Source/Niagara/Private/NiagaraDebuggerClient.cpp b/Engine/Plugins/FX/Niagara/Source/Niagara/Private/NiagaraDebuggerClient.cpp
--- a/Engine/Plugins/FX/Niagara/Source/Niagara/Private/NiagaraDebuggerClient.cpp
+++ b/Engine/Plugins/FX/Niagara/Source/Niagara/Private/NiagaraDebuggerClient.cpp
@@ -42,12 +42,12 @@ FNiagaraDebuggerClient::FNiagaraDebuggerClient()
 	InstanceName = FApp::GetInstanceName();
 	UE_LOG(LogNiagaraDebuggerClient, Log, TEXT("Niagara Debugger Client Initialized | Session: %s | Instance: %s (%s)."), *SessionId.ToString(), *InstanceId.ToString(), *InstanceName);
 
-	TickerHandle = FTicker::GetCoreTicker().AddTicker(FTickerDelegate::CreateRaw(this, &FNiagaraDebuggerClient::Tick));
+	TickerHandle = FTSTicker::GetCoreTicker().AddTicker(FTickerDelegate::CreateRaw(this, &FNiagaraDebuggerClient::Tick));
 }
 
 FNiagaraDebuggerClient::~FNiagaraDebuggerClient()
 {
-	FTicker::GetCoreTicker().RemoveTicker(TickerHandle);
+	FTSTicker::GetCoreTicker().RemoveTicker(TickerHandle);
 
 	CloseConnection();
 }
@@ -212,7 +212,7 @@ void FNiagaraDebuggerClient::HandleOutlinerSettingsMessage(const FNiagaraOutline
 				{
 					OutlinerCountdown = Message.CaptureDelayFrames;
 					UE_LOG(LogNiagaraDebuggerClient, Log, TEXT("Recieved request to capture outliner data. Capturing in %u frames. | Session: %s | Instance: %s (%s)."), Message.CaptureDelayFrames, *SessionId.ToString(), *InstanceId.ToString(), *InstanceName);
-					FTicker::GetCoreTicker().AddTicker(FTickerDelegate::CreateRaw(this, &FNiagaraDebuggerClient::UpdateOutliner));
+					FTSTicker::GetCoreTicker().AddTicker(FTickerDelegate::CreateRaw(this, &FNiagaraDebuggerClient::UpdateOutliner));
 				}
 			}
 			else
diff --git a/Engine/Plugins/FX/Niagara/Source/NiagaraEditor/Private/NiagaraEditorModule.cpp b/Engine/Plugins/FX/Niagara/Source/NiagaraEditor/Private/NiagaraEditorModule.cpp
--- a/Engine/Plugins/FX/Niagara/Source/NiagaraEditor/Private/NiagaraEditorModule.cpp
+++ b/Engine/Plugins/FX/Niagara/Source/NiagaraEditor/Private/NiagaraEditorModule.cpp
@@ -1706,7 +1706,7 @@ void FNiagaraEditorModule::EnqueueObjectForDeferredDestructionInternal(FDeferred
 {
 	if (EnqueuedForDeferredDestruction.Num() == 0)
 	{
-		FTicker::GetCoreTicker().AddTicker(FTickerDelegate::CreateRaw(this, &FNiagaraEditorModule::DeferredDestructObjects));
+		FTSTicker::GetCoreTicker().AddTicker(FTickerDelegate::CreateRaw(this, &FNiagaraEditorModule::DeferredDestructObjects));
 	}
 	EnqueuedForDeferredDestruction.Add(MakeShareable<FDeferredDestructionContainerBase>(InObjectToDestruct));
 }
diff --git a/Engine/Plugins/Lumin/MagicLeap/Source/MagicLeapEyeTracker/Private/MagicLeapVREyeTracker.h b/Engine/Plugins/Lumin/MagicLeap/Source/MagicLeapEyeTracker/Private/MagicLeapVREyeTracker.h
--- a/Engine/Plugins/Lumin/MagicLeap/Source/MagicLeapEyeTracker/Private/MagicLeapVREyeTracker.h
+++ b/Engine/Plugins/Lumin/MagicLeap/Source/MagicLeapEyeTracker/Private/MagicLeapVREyeTracker.h
@@ -12,7 +12,7 @@
 #include "Lumin/CAPIShims/LuminAPIEyeTracking.h"
 #include "IMagicLeapTrackerEntity.h"
 
-class FMagicLeapVREyeTracker : public IMagicLeapVREyeTracker, public FTickerObjectBase, public IMagicLeapTrackerEntity
+class FMagicLeapVREyeTracker : public IMagicLeapVREyeTracker, public FTSTickerObjectBase, public IMagicLeapTrackerEntity
 {
 public:
 	FMagicLeapVREyeTracker();
diff --git a/Engine/Plugins/Lumin/MagicLeap/Source/MagicLeapHandMeshing/Private/MagicLeapHandMeshingModule.cpp b/Engine/Plugins/Lumin/MagicLeap/Source/MagicLeapHandMeshing/Private/MagicLeapHandMeshingModule.cpp
--- a/Engine/Plugins/Lumin/MagicLeap/Source/MagicLeapHandMeshing/Private/MagicLeapHandMeshingModule.cpp
+++ b/Engine/Plugins/Lumin/MagicLeap/Source/MagicLeapHandMeshing/Private/MagicLeapHandMeshingModule.cpp
@@ -11,12 +11,12 @@ void FMagicLeapHandMeshingModule::StartupModule()
 {
 	IMagicLeapHandMeshingModule::StartupModule();
 	TickDelegate = FTickerDelegate::CreateRaw(this, &FMagicLeapHandMeshingModule::Tick);
-	TickDelegateHandle = FTicker::GetCoreTicker().AddTicker(TickDelegate);
+	TickDelegateHandle = FTSTicker::GetCoreTicker().AddTicker(TickDelegate);
 }
 
 void FMagicLeapHandMeshingModule::ShutdownModule()
 {
-	FTicker::GetCoreTicker().RemoveTicker(TickDelegateHandle);
+	FTSTicker::GetCoreTicker().RemoveTicker(TickDelegateHandle);
 	IModuleInterface::ShutdownModule();
 }
 
diff --git a/Engine/Plugins/Lumin/MagicLeap/Source/MagicLeapHandMeshing/Private/MagicLeapHandMeshingModule.h b/Engine/Plugins/Lumin/MagicLeap/Source/MagicLeapHandMeshing/Private/MagicLeapHandMeshingModule.h
--- a/Engine/Plugins/Lumin/MagicLeap/Source/MagicLeapHandMeshing/Private/MagicLeapHandMeshingModule.h
+++ b/Engine/Plugins/Lumin/MagicLeap/Source/MagicLeapHandMeshing/Private/MagicLeapHandMeshingModule.h
@@ -25,7 +25,7 @@ public:
 
 private:
 	FTickerDelegate TickDelegate;
-	FDelegateHandle TickDelegateHandle;
+	FTSTicker::FDelegateHandle TickDelegateHandle;
 	FMagicLeapMeshTracker MeshTracker;
 };
 
diff --git a/Engine/Plugins/Lumin/MagicLeap/Source/MagicLeapImageTracker/Private/MagicLeapImageTrackerModule.cpp b/Engine/Plugins/Lumin/MagicLeap/Source/MagicLeapImageTracker/Private/MagicLeapImageTrackerModule.cpp
--- a/Engine/Plugins/Lumin/MagicLeap/Source/MagicLeapImageTracker/Private/MagicLeapImageTrackerModule.cpp
+++ b/Engine/Plugins/Lumin/MagicLeap/Source/MagicLeapImageTracker/Private/MagicLeapImageTrackerModule.cpp
@@ -27,7 +27,7 @@ void FMagicLeapImageTrackerModule::ShutdownModule()
 {
 	if (TickDelegateHandle.IsValid())
 	{
-		FTicker::GetCoreTicker().RemoveTicker(TickDelegateHandle);
+		FTSTicker::GetCoreTicker().RemoveTicker(TickDelegateHandle);
 	}
 	FMagicLeapImageTrackerRunnable* InRunnable = Runnable;
 	Runnable = nullptr;
@@ -86,7 +86,7 @@ void FMagicLeapImageTrackerModule::SetTargetAsync(const FMagicLeapImageTargetSet
 	Runnable->SetTargetAsync(ImageTarget, SucceededDelegate, FailedDelegate);
 	if (!TickDelegateHandle.IsValid())
 	{
-		TickDelegateHandle = FTicker::GetCoreTicker().AddTicker(TickDelegate);
+		TickDelegateHandle = FTSTicker::GetCoreTicker().AddTicker(TickDelegate);
 	}
 }
 
@@ -95,7 +95,7 @@ void FMagicLeapImageTrackerModule::SetTargetAsync(const FMagicLeapImageTargetSet
 	Runnable->SetTargetAsync(ImageTarget, SucceededDelegate, FailedDelegate);
 	if (!TickDelegateHandle.IsValid())
 	{
-		TickDelegateHandle = FTicker::GetCoreTicker().AddTicker(TickDelegate);
+		TickDelegateHandle = FTSTicker::GetCoreTicker().AddTicker(TickDelegate);
 	}
 }
 
diff --git a/Engine/Plugins/Lumin/MagicLeap/Source/MagicLeapImageTracker/Private/MagicLeapImageTrackerModule.h b/Engine/Plugins/Lumin/MagicLeap/Source/MagicLeapImageTracker/Private/MagicLeapImageTrackerModule.h
--- a/Engine/Plugins/Lumin/MagicLeap/Source/MagicLeapImageTracker/Private/MagicLeapImageTrackerModule.h
+++ b/Engine/Plugins/Lumin/MagicLeap/Source/MagicLeapImageTracker/Private/MagicLeapImageTrackerModule.h
@@ -37,7 +37,7 @@ public:
 
 private:
 	FTickerDelegate TickDelegate;
-	FDelegateHandle TickDelegateHandle;
+	FTSTicker::FDelegateHandle TickDelegateHandle;
 	class FMagicLeapImageTrackerRunnable* Runnable;
 };
 
diff --git a/Engine/Plugins/Lumin/MagicLeap/Source/MagicLeapPlanes/Source/Private/MagicLeapPlanesModule.cpp b/Engine/Plugins/Lumin/MagicLeap/Source/MagicLeapPlanes/Source/Private/MagicLeapPlanesModule.cpp
--- a/Engine/Plugins/Lumin/MagicLeap/Source/MagicLeapPlanes/Source/Private/MagicLeapPlanesModule.cpp
+++ b/Engine/Plugins/Lumin/MagicLeap/Source/MagicLeapPlanes/Source/Private/MagicLeapPlanesModule.cpp
@@ -410,7 +410,7 @@ bool FMagicLeapPlanesModule::CreateTracker()
 
 		if (Result == MLResult_Ok)
 		{
-			TickDelegateHandle = FTicker::GetCoreTicker().AddTicker(TickDelegate);
+			TickDelegateHandle = FTSTicker::GetCoreTicker().AddTicker(TickDelegate);
 		}
 
 		return Result == MLResult_Ok;
@@ -435,7 +435,7 @@ bool FMagicLeapPlanesModule::DestroyTracker()
 		MLResult Result = MLPlanesDestroy(Tracker);
 		UE_CLOG(Result != MLResult_Ok, LogMagicLeapPlanes, Error, TEXT("MLPlanesDestroy failed with error '%s'."), UTF8_TO_TCHAR(MLGetResultString(Result)));
 		Tracker = ML_INVALID_HANDLE;
-		FTicker::GetCoreTicker().RemoveTicker(TickDelegateHandle);
+		FTSTicker::GetCoreTicker().RemoveTicker(TickDelegateHandle);
 		return true;
 	}
 	else
diff --git a/Engine/Plugins/Lumin/MagicLeapCamera/Source/MagicLeapCVCamera/Private/MagicLeapCVCameraModule.cpp b/Engine/Plugins/Lumin/MagicLeapCamera/Source/MagicLeapCVCamera/Private/MagicLeapCVCameraModule.cpp
--- a/Engine/Plugins/Lumin/MagicLeapCamera/Source/MagicLeapCVCamera/Private/MagicLeapCVCameraModule.cpp
+++ b/Engine/Plugins/Lumin/MagicLeapCamera/Source/MagicLeapCVCamera/Private/MagicLeapCVCameraModule.cpp
@@ -35,7 +35,7 @@ void FMagicLeapCVCameraModule::ShutdownModule()
 	{
 		delete InRunnable;
 	});
-	FTicker::GetCoreTicker().RemoveTicker(TickDelegateHandle);
+	FTSTicker::GetCoreTicker().RemoveTicker(TickDelegateHandle);
 	IModuleInterface::ShutdownModule();
 }
 
@@ -86,7 +86,7 @@ bool FMagicLeapCVCameraModule::EnableAsync(const FMagicLeapCVCameraEnableStatic&
 {
 	if (TryPushNewCaptureTask(FCVCameraTask::EType::Enable))
 	{
-		TickDelegateHandle = FTicker::GetCoreTicker().AddTicker(TickDelegate);
+		TickDelegateHandle = FTSTicker::GetCoreTicker().AddTicker(TickDelegate);
 		static constexpr float MinCameraRunDuration = 2.0f;
 		MinCameraRunTimer = MinCameraRunDuration;
 		OnEnabledStatic = ResultDelegate;
@@ -100,7 +100,7 @@ bool FMagicLeapCVCameraModule::EnableAsync(const FMagicLeapCVCameraEnableMulti&
 {
 	if (TryPushNewCaptureTask(FCVCameraTask::EType::Enable))
 	{
-		TickDelegateHandle = FTicker::GetCoreTicker().AddTicker(TickDelegate);
+		TickDelegateHandle = FTSTicker::GetCoreTicker().AddTicker(TickDelegate);
 		static constexpr float MinCameraRunDuration = 2.0f;
 		MinCameraRunTimer = MinCameraRunDuration;
 		OnEnabled = ResultDelegate;
@@ -114,7 +114,7 @@ bool FMagicLeapCVCameraModule::DisableAsync(const FMagicLeapCVCameraDisableStati
 {
 	if (MinCameraRunTimer < 0.0f && TryPushNewCaptureTask(FCVCameraTask::EType::Disable))
 	{
-		FTicker::GetCoreTicker().RemoveTicker(TickDelegateHandle);
+		FTSTicker::GetCoreTicker().RemoveTicker(TickDelegateHandle);
 		OnDisabledStatic = ResultDelegate;
 		return true;
 	}
@@ -126,7 +126,7 @@ bool FMagicLeapCVCameraModule::DisableAsync(const FMagicLeapCVCameraDisableMulti
 {
 	if (MinCameraRunTimer < 0.0f && TryPushNewCaptureTask(FCVCameraTask::EType::Disable))
 	{
-		FTicker::GetCoreTicker().RemoveTicker(TickDelegateHandle);
+		FTSTicker::GetCoreTicker().RemoveTicker(TickDelegateHandle);
 		OnDisabled = ResultDelegate;
 		return true;
 	}
diff --git a/Engine/Plugins/Lumin/MagicLeapDispatch/Source/Private/MagicLeapDispatchPlugin.cpp b/Engine/Plugins/Lumin/MagicLeapDispatch/Source/Private/MagicLeapDispatchPlugin.cpp
--- a/Engine/Plugins/Lumin/MagicLeapDispatch/Source/Private/MagicLeapDispatchPlugin.cpp
+++ b/Engine/Plugins/Lumin/MagicLeapDispatch/Source/Private/MagicLeapDispatchPlugin.cpp
@@ -44,7 +44,7 @@ public:
 
 	void ShutdownModule() override
 	{
-		FTicker::GetCoreTicker().RemoveTicker(TickDelegateHandle);
+		FTSTicker::GetCoreTicker().RemoveTicker(TickDelegateHandle);
 	}
 
 	bool Tick(float DeltaTime)
@@ -121,7 +121,7 @@ public:
 #if WITH_MLSDK
 		if (!TickDelegateHandle.IsValid())
 		{
-			TickDelegateHandle = FTicker::GetCoreTicker().AddTicker(TickDelegate);
+			TickDelegateHandle = FTSTicker::GetCoreTicker().AddTicker(TickDelegate);
 		}
 
 		OAuthMetadata.RedirectURI = RedirectURI;
@@ -222,7 +222,7 @@ public:
 
 private:
 	FTickerDelegate TickDelegate;
-	FDelegateHandle TickDelegateHandle;
+	FTSTicker::FDelegateHandle TickDelegateHandle;
 };
 
 IMPLEMENT_MODULE(FMagicLeapDispatchPlugin, MagicLeapDispatch);
diff --git a/Engine/Plugins/Lumin/MagicLeapSharedFile/Source/Private/MagicLeapSharedFilePlugin.cpp b/Engine/Plugins/Lumin/MagicLeapSharedFile/Source/Private/MagicLeapSharedFilePlugin.cpp
--- a/Engine/Plugins/Lumin/MagicLeapSharedFile/Source/Private/MagicLeapSharedFilePlugin.cpp
+++ b/Engine/Plugins/Lumin/MagicLeapSharedFile/Source/Private/MagicLeapSharedFilePlugin.cpp
@@ -67,7 +67,7 @@ void FMagicLeapSharedFilePlugin::StartupModule()
 
 void FMagicLeapSharedFilePlugin::ShutdownModule()
 {
-	FTicker::GetCoreTicker().RemoveTicker(TickDelegateHandle);
+	FTSTicker::GetCoreTicker().RemoveTicker(TickDelegateHandle);
 	IMagicLeapSharedFilePlugin::ShutdownModule();
 }
 
@@ -119,14 +119,14 @@ bool FMagicLeapSharedFilePlugin::SharedFilePickAsync(const FMagicLeapFilesPicked
 		bWaitingForDelegateResult = true;
 	}
 
-	TickDelegateHandle = FTicker::GetCoreTicker().AddTicker(TickDelegate);
+	TickDelegateHandle = FTSTicker::GetCoreTicker().AddTicker(TickDelegate);
 
 	MLResult Result = MLSharedFilePick(onFilesPicked, this);
 	if (MLResult_Ok != Result)
 	{
 		FScopeLock Lock(&Mutex);
 		bWaitingForDelegateResult = true;
-		FTicker::GetCoreTicker().RemoveTicker(TickDelegateHandle);
+		FTSTicker::GetCoreTicker().RemoveTicker(TickDelegateHandle);
 		UE_LOG(LogMagicLeapSharedFile, Error, TEXT("MLSharedFilePick failed with error '%s'"), UTF8_TO_TCHAR(MLSharedFileGetResultString(Result)));
 	}
 
@@ -151,7 +151,7 @@ bool FMagicLeapSharedFilePlugin::Tick(float DeltaTime)
 	
 	if (!bWaitingForDelegateResult_Cached)
 	{
-		FTicker::GetCoreTicker().RemoveTicker(TickDelegateHandle);
+		FTSTicker::GetCoreTicker().RemoveTicker(TickDelegateHandle);
 		ResultDelegate_Cached.ExecuteIfBound(PickedFileList);
 	}
 
diff --git a/Engine/Plugins/Messaging/UdpMessaging/Source/UdpMessaging/Private/Transport/UdpMessageTransport.cpp b/Engine/Plugins/Messaging/UdpMessaging/Source/UdpMessaging/Private/Transport/UdpMessageTransport.cpp
--- a/Engine/Plugins/Messaging/UdpMessaging/Source/UdpMessaging/Private/Transport/UdpMessageTransport.cpp
+++ b/Engine/Plugins/Messaging/UdpMessaging/Source/UdpMessaging/Private/Transport/UdpMessageTransport.cpp
@@ -316,7 +316,7 @@ void FUdpMessageTransport::StartAutoRepairRoutine(uint32 MaxRetryAttempt)
 	TWeakPtr<FUdpMessageTransport, ESPMode::ThreadSafe> WeakTransportPtr = AsShared();
 	FTimespan CheckDelay(0, 0, 1);
 	uint32 CheckNumber = 1;
-	AutoRepairHandle = FTicker::GetCoreTicker().AddTicker(FTickerDelegate::CreateLambda([WeakTransport = WeakTransportPtr, LastTime = FDateTime::UtcNow(), CheckDelay, CheckNumber, MaxRetryAttempt](float DeltaTime) mutable
+	AutoRepairHandle = FTSTicker::GetCoreTicker().AddTicker(FTickerDelegate::CreateLambda([WeakTransport = WeakTransportPtr, LastTime = FDateTime::UtcNow(), CheckDelay, CheckNumber, MaxRetryAttempt](float DeltaTime) mutable
 	{
 		QUICK_SCOPE_CYCLE_COUNTER(STAT_FUdpMessageTransport_AutoRepair);
 		bool bContinue = true;
@@ -347,7 +347,7 @@ void FUdpMessageTransport::StopAutoRepairRoutine()
 {
 	if (AutoRepairHandle.IsValid())
 	{
-		FTicker::GetCoreTicker().RemoveTicker(AutoRepairHandle);
+		FTSTicker::GetCoreTicker().RemoveTicker(AutoRepairHandle);
 	}
 }
 
diff --git a/Engine/Plugins/Messaging/UdpMessaging/Source/UdpMessaging/Private/Transport/UdpMessageTransport.h b/Engine/Plugins/Messaging/UdpMessaging/Source/UdpMessaging/Private/Transport/UdpMessageTransport.h
--- a/Engine/Plugins/Messaging/UdpMessaging/Source/UdpMessaging/Private/Transport/UdpMessageTransport.h
+++ b/Engine/Plugins/Messaging/UdpMessaging/Source/UdpMessaging/Private/Transport/UdpMessageTransport.h
@@ -9,6 +9,7 @@
 #include "IMessageAttachment.h"
 #include "Templates/SharedPointer.h"
 #include "Async/Future.h"
+#include "Containers/Ticker.h"
 
 class FArrayReader;
 class FUdpReassembledMessage;
@@ -105,7 +106,7 @@ private:
 	TFuture<void> ErrorFuture;
 
 	/** Holds the delegate handle for the auto repair routine. */
-	FDelegateHandle AutoRepairHandle;
+	FTSTicker::FDelegateHandle AutoRepairHandle;
 
 	/** Holds the message processor. */
 	FUdpMessageProcessor* MessageProcessor;
diff --git a/Engine/Plugins/NetcodeUnitTest/NetcodeUnitTest/Source/NetcodeUnitTest/Private/UnitTestCommandlet.cpp b/Engine/Plugins/NetcodeUnitTest/NetcodeUnitTest/Source/NetcodeUnitTest/Private/UnitTestCommandlet.cpp
--- a/Engine/Plugins/NetcodeUnitTest/NetcodeUnitTest/Source/NetcodeUnitTest/Private/UnitTestCommandlet.cpp
+++ b/Engine/Plugins/NetcodeUnitTest/NetcodeUnitTest/Source/NetcodeUnitTest/Private/UnitTestCommandlet.cpp
@@ -159,8 +159,7 @@ int32 UUnitTestCommandlet::Main(const FString& Params)
 			// Required for FTimerManager to function - as it blocks ticks, if the frame counter doesn't change
 			GFrameCounter++;
 
-			FTicker::GetCoreTicker().Tick(FApp::GetDeltaTime());
-
+			FTSTicker::GetCoreTicker().Tick(FApp::GetDeltaTime());
 
 			// Execute deferred commands
 			for (int32 DeferredCommandsIndex=0; DeferredCommandsIndex<GEngine->DeferredCommands.Num(); DeferredCommandsIndex++)
diff --git a/Engine/Plugins/Online/EOSShared/Source/EOSShared/Private/EOSSDKManager.cpp b/Engine/Plugins/Online/EOSShared/Source/EOSShared/Private/EOSSDKManager.cpp
--- a/Engine/Plugins/Online/EOSShared/Source/EOSShared/Private/EOSSDKManager.cpp
+++ b/Engine/Plugins/Online/EOSShared/Source/EOSShared/Private/EOSSDKManager.cpp
@@ -206,7 +206,7 @@ IEOSPlatformHandlePtr FEOSSDKManager::CreatePlatform(const EOS_Platform_Options&
 
 			if (!TickerHandle.IsValid())
 			{
-				TickerHandle = FTicker::GetCoreTicker().AddTicker(FTickerDelegate::CreateRaw(this, &FEOSSDKManager::Tick), 0.0f);
+				TickerHandle = FTSTicker::GetCoreTicker().AddTicker(FTickerDelegate::CreateRaw(this, &FEOSSDKManager::Tick), 0.0f);
 			}
 		}
 		else
@@ -270,7 +270,7 @@ void FEOSSDKManager::ReleasePlatform(EOS_HPlatform PlatformHandle)
 		if (ensure(TickerHandle.IsValid()) &&
 			PlatformHandles.Num() == 0)
 		{
-			FTicker::GetCoreTicker().RemoveTicker(TickerHandle);
+			FTSTicker::GetCoreTicker().RemoveTicker(TickerHandle);
 			TickerHandle.Reset();
 		}
 	}
@@ -295,7 +295,7 @@ void FEOSSDKManager::Shutdown()
 
 			if (ensure(TickerHandle.IsValid()))
 			{
-				FTicker::GetCoreTicker().RemoveTicker(TickerHandle);
+				FTSTicker::GetCoreTicker().RemoveTicker(TickerHandle);
 				TickerHandle.Reset();
 			}
 		}
diff --git a/Engine/Plugins/Online/EOSShared/Source/EOSShared/Private/EOSSDKManager.h b/Engine/Plugins/Online/EOSShared/Source/EOSShared/Private/EOSSDKManager.h
--- a/Engine/Plugins/Online/EOSShared/Source/EOSShared/Private/EOSSDKManager.h
+++ b/Engine/Plugins/Online/EOSShared/Source/EOSShared/Private/EOSSDKManager.h
@@ -13,6 +13,7 @@
 #include "eos_init.h"
 
 #include "IEOSSDKManager.h"
+#include "Containers/Ticker.h"
 
 struct FEOSPlatformHandle;
 
@@ -53,7 +54,7 @@ private:
 	/** Created platforms */
 	TArray<EOS_HPlatform> PlatformHandles;
 	/** Handle to ticker delegate for Tick(), valid whenever EosPlatformHandles is non-empty. */
-	FDelegateHandle TickerHandle;
+	FTSTicker::FDelegateHandle TickerHandle;
 };
 
 struct FEOSPlatformHandle : public IEOSPlatformHandle
diff --git a/Engine/Plugins/Online/Experimental/OnlineServices/Source/OnlineServicesCommon/Public/Online/OnlineServicesCommon.h b/Engine/Plugins/Online/Experimental/OnlineServices/Source/OnlineServicesCommon/Public/Online/OnlineServicesCommon.h
--- a/Engine/Plugins/Online/Experimental/OnlineServices/Source/OnlineServicesCommon/Public/Online/OnlineServicesCommon.h
+++ b/Engine/Plugins/Online/Experimental/OnlineServices/Source/OnlineServicesCommon/Public/Online/OnlineServicesCommon.h
@@ -18,7 +18,7 @@ namespace UE::Online {
 class ONLINESERVICESCOMMON_API FOnlineServicesCommon
 	: public IOnlineServices
 	, public TSharedFromThis<FOnlineServicesCommon>
-	, public FTickerObjectBase
+	, public FTSTickerObjectBase
 {
 public:
 	using Super = IOnlineServices;
diff --git a/Engine/Plugins/Online/IOS/OnlineSubsystemIOS/Source/Private/OnlinePurchaseIOS.cpp b/Engine/Plugins/Online/IOS/OnlineSubsystemIOS/Source/Private/OnlinePurchaseIOS.cpp
--- a/Engine/Plugins/Online/IOS/OnlineSubsystemIOS/Source/Private/OnlinePurchaseIOS.cpp
+++ b/Engine/Plugins/Online/IOS/OnlineSubsystemIOS/Source/Private/OnlinePurchaseIOS.cpp
@@ -393,7 +393,7 @@ void FOnlinePurchaseIOS::OnTransactionDeferred(const FStoreKitTransactionData& T
 #if TEST_DEFERRED_TRANSACTIONS
 		// Existing transactions emulating deferment will trigger as success shortly
 		TWeakPtr<FOnlinePurchaseIOS, ESPMode::ThreadSafe> WeakThis(AsShared());
-		FTicker::GetCoreTicker().AddTicker(FTickerDelegate::CreateLambda([WeakThis, TransactionData](float) -> bool
+		FTSTicker::GetCoreTicker().AddTicker(FTickerDelegate::CreateLambda([WeakThis, TransactionData](float) -> bool
 		{
 			QUICK_SCOPE_CYCLE_COUNTER(STAT_FOnlinePurchaseIOS_TestDeferredTransactions);
 			FOnlinePurchaseIOSPtr StrongThis = WeakThis.Pin();
diff --git a/Engine/Plugins/Online/OnlineFramework/Source/Hotfix/Private/UpdateManager.cpp b/Engine/Plugins/Online/OnlineFramework/Source/Hotfix/Private/UpdateManager.cpp
--- a/Engine/Plugins/Online/OnlineFramework/Source/Hotfix/Private/UpdateManager.cpp
+++ b/Engine/Plugins/Online/OnlineFramework/Source/Hotfix/Private/UpdateManager.cpp
@@ -510,7 +510,7 @@ void UUpdateManager::StartInitialPreload()
 	SetUpdateState(EUpdateState::WaitingOnInitialLoad);
 
 	// Start ticking
-	FTicker& Ticker = FTicker::GetCoreTicker();
+	FTSTicker& Ticker = FTSTicker::GetCoreTicker();
 	FTickerDelegate TickDelegate = FTickerDelegate::CreateUObject(this, &ThisClass::Tick);
 	ensure(!TickerHandle.IsValid());
 	TickerHandle = Ticker.AddTicker(TickDelegate, 0.0f);
diff --git a/Engine/Plugins/Online/OnlineFramework/Source/Party/Private/Party/PartyPlatformSessionMonitor.cpp b/Engine/Plugins/Online/OnlineFramework/Source/Party/Private/Party/PartyPlatformSessionMonitor.cpp
--- a/Engine/Plugins/Online/OnlineFramework/Source/Party/Private/Party/PartyPlatformSessionMonitor.cpp
+++ b/Engine/Plugins/Online/OnlineFramework/Source/Party/Private/Party/PartyPlatformSessionMonitor.cpp
@@ -220,7 +220,7 @@ bool FPartyPlatformSessionManager::FindSessionInternal(const FSessionId& Session
 					UE_LOG(LogParty, Warning, TEXT("PartyPlatformSessionMonitor adding artificial delay of %0.2fs to session find attempt"), DelaySeconds);
 
 					TWeakPtr<FPartyPlatformSessionManager> AsWeakPtr = SharedThis(this);
-					FTicker::GetCoreTicker().AddTicker(FTickerDelegate::CreateLambda(
+					FTSTicker::GetCoreTicker().AddTicker(FTickerDelegate::CreateLambda(
 						[AsWeakPtr, SessionId, SessionIdString, SessionOwnerId, LocalUserPlatformId, OnAttemptComplete, this](float)
 						{
 							QUICK_SCOPE_CYCLE_COUNTER(STAT_FPartyPlatformSessionManager_FindSessionAttempt);
@@ -515,7 +515,7 @@ void FPartyPlatformSessionMonitor::ShutdownInternal()
 
 	if (RetryTickerHandle.IsValid())
 	{
-		FTicker::GetCoreTicker().RemoveTicker(RetryTickerHandle);
+		FTSTicker::GetCoreTicker().RemoveTicker(RetryTickerHandle);
 		RetryTickerHandle.Reset();
 	}
 
@@ -545,7 +545,7 @@ void FPartyPlatformSessionMonitor::CreateSession(const FUniqueNetIdRepl& LocalUs
 			UE_LOG(LogParty, Warning, TEXT("PartyPlatformSessionMonitor adding artificial delay of %0.2fs to session creation attempt"), DelaySeconds);
 
 			TWeakPtr<FPartyPlatformSessionMonitor> AsWeakPtr = SharedThis(this);
-			FTicker::GetCoreTicker().AddTicker(FTickerDelegate::CreateLambda(
+			FTSTicker::GetCoreTicker().AddTicker(FTickerDelegate::CreateLambda(
 				[AsWeakPtr, SessionSettings, LocalUserPlatformId, this] (float)
 				{
 					QUICK_SCOPE_CYCLE_COUNTER(STAT_FPartyPlatformSessionManager_CreateSessionAttempt);
@@ -668,7 +668,7 @@ void FPartyPlatformSessionMonitor::JoinSession(const FOnlineSessionSearchResult&
 		UE_LOG(LogParty, Warning, TEXT("Adding artificial delay of %0.2fs to session join attempt"), DelaySeconds);
 
 		TWeakPtr<FPartyPlatformSessionMonitor> AsWeakPtr = SharedThis(this);
-		FTicker::GetCoreTicker().AddTicker(FTickerDelegate::CreateLambda(
+		FTSTicker::GetCoreTicker().AddTicker(FTickerDelegate::CreateLambda(
 			[AsWeakPtr, SearchResultCopy, LocalUserPlatformId, this] (float)
 			{
 				QUICK_SCOPE_CYCLE_COUNTER(STAT_FPartyPlatformSessionManager_JoinSessionAttempt);
@@ -719,7 +719,7 @@ void FPartyPlatformSessionMonitor::QueuePlatformSessionUpdate()
 	{
 		UE_LOG(LogParty, Verbose, TEXT("PartyPlatformSessionMonitor queuing session update for party [%s]"), *MonitoredParty->ToDebugString())
 		bHasQueuedSessionUpdate = true;
-		FTicker::GetCoreTicker().AddTicker(FTickerDelegate::CreateSP(this, &FPartyPlatformSessionMonitor::HandleQueuedSessionUpdate));
+		FTSTicker::GetCoreTicker().AddTicker(FTickerDelegate::CreateSP(this, &FPartyPlatformSessionMonitor::HandleQueuedSessionUpdate));
 	}
 }
 
@@ -903,7 +903,7 @@ void FPartyPlatformSessionMonitor::HandleCreateSessionComplete(FName SessionName
 		if (AllowCreateSessionFailure == 0 && ensure(!RetryTickerHandle.IsValid()))
 		{
 			// Unsuccessful and we aren't trying to leave, so we'll try again here in a moment
-			RetryTickerHandle = FTicker::GetCoreTicker().AddTicker(FTickerDelegate::CreateSP(this, &FPartyPlatformSessionMonitor::HandleRetryEstablishingSession), EstablishSessionRetryDelay);
+			RetryTickerHandle = FTSTicker::GetCoreTicker().AddTicker(FTickerDelegate::CreateSP(this, &FPartyPlatformSessionMonitor::HandleRetryEstablishingSession), EstablishSessionRetryDelay);
 		}
 	}
 }
@@ -1137,7 +1137,7 @@ void FPartyPlatformSessionMonitor::ProcessJoinFailure()
 	else
 	{
 		TargetSessionId = FSessionId();
-		RetryTickerHandle = FTicker::GetCoreTicker().AddTicker(FTickerDelegate::CreateSP(this, &FPartyPlatformSessionMonitor::HandleRetryEstablishingSession), EstablishSessionRetryDelay);
+		RetryTickerHandle = FTSTicker::GetCoreTicker().AddTicker(FTickerDelegate::CreateSP(this, &FPartyPlatformSessionMonitor::HandleRetryEstablishingSession), EstablishSessionRetryDelay);
 	}
 }
 
@@ -1184,6 +1184,6 @@ void FPartyPlatformSessionMonitor::HandleSessionFailure(const FUniqueNetId& Loca
 	{
 		MonitoredParty->SetIsMissingPlatformSession(true);
 
-		RetryTickerHandle = FTicker::GetCoreTicker().AddTicker(FTickerDelegate::CreateSP(this, &FPartyPlatformSessionMonitor::HandleRetryEstablishingSession), EstablishSessionRetryDelay);
+		RetryTickerHandle = FTSTicker::GetCoreTicker().AddTicker(FTickerDelegate::CreateSP(this, &FPartyPlatformSessionMonitor::HandleRetryEstablishingSession), EstablishSessionRetryDelay);
 	}
 }
diff --git a/Engine/Plugins/Online/OnlineFramework/Source/Party/Private/Party/PartyPlatformSessionMonitor.h b/Engine/Plugins/Online/OnlineFramework/Source/Party/Private/Party/PartyPlatformSessionMonitor.h
--- a/Engine/Plugins/Online/OnlineFramework/Source/Party/Private/Party/PartyPlatformSessionMonitor.h
+++ b/Engine/Plugins/Online/OnlineFramework/Source/Party/Private/Party/PartyPlatformSessionMonitor.h
@@ -138,6 +138,6 @@ private:
 		Complete
 	} ShutdownState = EMonitorShutdownState::None;
 	
-	FDelegateHandle RetryTickerHandle;
+	FTSTicker::FDelegateHandle RetryTickerHandle;
 	FDelegateHandle UpdateOssSessionHandle;
 };
\ No newline at end of file
diff --git a/Engine/Plugins/Online/OnlineFramework/Source/Party/Private/SocialToolkit.cpp b/Engine/Plugins/Online/OnlineFramework/Source/Party/Private/SocialToolkit.cpp
--- a/Engine/Plugins/Online/OnlineFramework/Source/Party/Private/SocialToolkit.cpp
+++ b/Engine/Plugins/Online/OnlineFramework/Source/Party/Private/SocialToolkit.cpp
@@ -1225,14 +1225,14 @@ void USocialToolkit::Debug_OnStartRandomizeUserPresence(uint8 NumRandomUser, flo
 {
 	if (Debug_PresenceTickerHandle.IsValid())
 	{
-		FTicker::GetCoreTicker().RemoveTicker(Debug_PresenceTickerHandle);
+		FTSTicker::GetCoreTicker().RemoveTicker(Debug_PresenceTickerHandle);
 		Debug_PresenceTickerHandle.Reset();
 	}
 
 	if (ensure(TickerTimer > 0.f))
 	{
 		bDebug_IsRandomlyChangingUserPresence = true;
-		Debug_PresenceTickerHandle = FTicker::GetCoreTicker().AddTicker(FTickerDelegate::CreateUObject(this, &USocialToolkit::Debug_HandleRandomizeUserPresenceTick, NumRandomUser), TickerTimer);
+		Debug_PresenceTickerHandle = FTSTicker::GetCoreTicker().AddTicker(FTickerDelegate::CreateUObject(this, &USocialToolkit::Debug_HandleRandomizeUserPresenceTick, NumRandomUser), TickerTimer);
 	}
 }
 
@@ -1241,7 +1241,7 @@ void USocialToolkit::Debug_OnStopRandomizeUserPresence(bool bClearGeneratedPrese
 	bDebug_IsRandomlyChangingUserPresence = false;
 	if (Debug_PresenceTickerHandle.IsValid())
 	{
-		FTicker::GetCoreTicker().RemoveTicker(Debug_PresenceTickerHandle);
+		FTSTicker::GetCoreTicker().RemoveTicker(Debug_PresenceTickerHandle);
 		Debug_PresenceTickerHandle.Reset();
 	}
 
diff --git a/Engine/Plugins/Online/OnlineFramework/Source/Party/Private/User/SocialUserList.cpp b/Engine/Plugins/Online/OnlineFramework/Source/Party/Private/User/SocialUserList.cpp
--- a/Engine/Plugins/Online/OnlineFramework/Source/Party/Private/User/SocialUserList.cpp
+++ b/Engine/Plugins/Online/OnlineFramework/Source/Party/Private/User/SocialUserList.cpp
@@ -106,12 +106,12 @@ void FSocialUserList::SetAllowAutoUpdate(bool bIsEnabled)
 
 	if (AutoUpdateRequests == 0 && UpdateTickerHandle.IsValid())
 	{
-		FTicker::GetCoreTicker().RemoveTicker(UpdateTickerHandle);
+		FTSTicker::GetCoreTicker().RemoveTicker(UpdateTickerHandle);
 		UpdateTickerHandle.Reset();
 	}
 	else if (bIsEnabled && !UpdateTickerHandle.IsValid())
 	{
-		UpdateTickerHandle = FTicker::GetCoreTicker().AddTicker(FTickerDelegate::CreateSP(this, &FSocialUserList::HandleAutoUpdateList), AutoUpdatePeriod);
+		UpdateTickerHandle = FTSTicker::GetCoreTicker().AddTicker(FTickerDelegate::CreateSP(this, &FSocialUserList::HandleAutoUpdateList), AutoUpdatePeriod);
 	}
 }
 
diff --git a/Engine/Plugins/Online/OnlineFramework/Source/Party/Private/User/SocialUserList.h b/Engine/Plugins/Online/OnlineFramework/Source/Party/Private/User/SocialUserList.h
--- a/Engine/Plugins/Online/OnlineFramework/Source/Party/Private/User/SocialUserList.h
+++ b/Engine/Plugins/Online/OnlineFramework/Source/Party/Private/User/SocialUserList.h
@@ -5,6 +5,7 @@
 #include "User/ISocialUserList.h"
 #include "UObject/GCObject.h"
 #include "PartyPackage.h"
+#include "Containers/Ticker.h"
 
 enum class EMemberExitedReason : uint8;
 typedef TSharedRef<const IOnlinePartyRequestToJoinInfo> IOnlinePartyRequestToJoinInfoConstRef;
@@ -103,7 +104,7 @@ private:
 	bool bNeedsSort = false;
 	int32 AutoUpdateRequests = 0;
 	float AutoUpdatePeriod = .5f;
-	FDelegateHandle UpdateTickerHandle;
+	FTSTicker::FDelegateHandle UpdateTickerHandle;
 
 	mutable FOnUserAdded OnUserAddedEvent;
 	mutable FOnUserRemoved OnUserRemovedEvent;
diff --git a/Engine/Plugins/Online/OnlineFramework/Source/Party/Public/Party/PartyDataReplicator.h b/Engine/Plugins/Online/OnlineFramework/Source/Party/Public/Party/PartyDataReplicator.h
--- a/Engine/Plugins/Online/OnlineFramework/Source/Party/Public/Party/PartyDataReplicator.h
+++ b/Engine/Plugins/Online/OnlineFramework/Source/Party/Public/Party/PartyDataReplicator.h
@@ -65,7 +65,7 @@ public:
 		if (UpdateTickerHandle.IsValid())
 		{
 			// Running manually - unregister ticker.
-			FTicker::GetCoreTicker().RemoveTicker(UpdateTickerHandle);
+			FTSTicker::GetCoreTicker().RemoveTicker(UpdateTickerHandle);
 			UpdateTickerHandle.Reset();
 			DeferredHandleReplicateChanges(0.f);
 		}
@@ -110,7 +110,7 @@ PACKAGE_SCOPE:
 		}
 		if (UpdateTickerHandle.IsValid())
 		{
-			FTicker::GetCoreTicker().RemoveTicker(UpdateTickerHandle);
+			FTSTicker::GetCoreTicker().RemoveTicker(UpdateTickerHandle);
 			UpdateTickerHandle.Reset();
 		}
 	}
@@ -120,7 +120,7 @@ private:
 	{
 		if (!UpdateTickerHandle.IsValid())
 		{
-			UpdateTickerHandle = FTicker::GetCoreTicker().AddTicker(FTickerDelegate::CreateRaw(this, &TPartyDataReplicator::DeferredHandleReplicateChanges));
+			UpdateTickerHandle = FTSTicker::GetCoreTicker().AddTicker(FTickerDelegate::CreateRaw(this, &TPartyDataReplicator::DeferredHandleReplicateChanges));
 		}
 	}
 
@@ -154,5 +154,5 @@ private:
 	/** Scratch copy of child UStruct for handling replication comparisons */
 	RepDataT* RepDataCopy = nullptr;
 
-	FDelegateHandle UpdateTickerHandle;
+	FTSTicker::FDelegateHandle UpdateTickerHandle;
 };
\ No newline at end of file
diff --git a/Engine/Plugins/Online/OnlineFramework/Source/Party/Public/SocialQuery.h b/Engine/Plugins/Online/OnlineFramework/Source/Party/Public/SocialQuery.h
--- a/Engine/Plugins/Online/OnlineFramework/Source/Party/Public/SocialQuery.h
+++ b/Engine/Plugins/Online/OnlineFramework/Source/Party/Public/SocialQuery.h
@@ -147,7 +147,7 @@ private:
 		// If we aren't already registered to execute our queries next tick, do so now
 		if (!TickExecuteHandle.IsValid())
 		{
-			TickExecuteHandle = FTicker::GetCoreTicker().AddTicker(FTickerDelegate::CreateRaw(this, &FSocialQueryManager::HandleExecuteQueries), UserInfoQueryAggregationTime);
+			TickExecuteHandle = FTSTicker::GetCoreTicker().AddTicker(FTickerDelegate::CreateRaw(this, &FSocialQueryManager::HandleExecuteQueries), UserInfoQueryAggregationTime);
 		}
 
 		return NewQuery;
@@ -161,6 +161,6 @@ private:
 		}
 	}
 
-	FDelegateHandle TickExecuteHandle;
+	FTSTicker::FDelegateHandle TickExecuteHandle;
 	TMap<FName, TArray<TSharedRef<FSocialQueryBase>>> CurrentQueriesById;
 };
\ No newline at end of file
diff --git a/Engine/Plugins/Online/OnlineFramework/Source/Party/Public/SocialToolkit.h b/Engine/Plugins/Online/OnlineFramework/Source/Party/Public/SocialToolkit.h
--- a/Engine/Plugins/Online/OnlineFramework/Source/Party/Public/SocialToolkit.h
+++ b/Engine/Plugins/Online/OnlineFramework/Source/Party/Public/SocialToolkit.h
@@ -268,7 +268,7 @@ private:	// Handlers
 	bool Debug_HandleRandomizeUserPresenceTick(float DeltaTime, uint8 NumRandomUser);
 	void Debug_ChangeRandomUserPresence(uint8 NumRandomUser);
 	bool bDebug_IsRandomlyChangingUserPresence = false;
-	FDelegateHandle Debug_PresenceTickerHandle;
+	FTSTicker::FDelegateHandle Debug_PresenceTickerHandle;
 #endif
 
 private:
diff --git a/Engine/Plugins/Online/OnlineFramework/Source/PlayTimeLimit/Private/PlayTimeLimitImpl.cpp b/Engine/Plugins/Online/OnlineFramework/Source/PlayTimeLimit/Private/PlayTimeLimitImpl.cpp
--- a/Engine/Plugins/Online/OnlineFramework/Source/PlayTimeLimit/Private/PlayTimeLimitImpl.cpp
+++ b/Engine/Plugins/Online/OnlineFramework/Source/PlayTimeLimit/Private/PlayTimeLimitImpl.cpp
@@ -43,7 +43,7 @@ void FPlayTimeLimitImpl::Initialize()
 	{
 		// Register delegate for ticker callback
 		FTickerDelegate TickDelegate = FTickerDelegate::CreateRaw(this, &FPlayTimeLimitImpl::Tick);
-		TickHandle = FTicker::GetCoreTicker().AddTicker(TickDelegate, 0.0f);
+		TickHandle = FTSTicker::GetCoreTicker().AddTicker(TickDelegate, 0.0f);
 	}
 }
 
@@ -53,7 +53,7 @@ void FPlayTimeLimitImpl::Shutdown()
 
 	if (TickHandle.IsValid())
 	{
-		FTicker::GetCoreTicker().RemoveTicker(TickHandle);
+		FTSTicker::GetCoreTicker().RemoveTicker(TickHandle);
 		TickHandle.Reset();
 	}
 
diff --git a/Engine/Plugins/Online/OnlineFramework/Source/PlayTimeLimit/Public/PlayTimeLimitImpl.h b/Engine/Plugins/Online/OnlineFramework/Source/PlayTimeLimit/Public/PlayTimeLimitImpl.h
--- a/Engine/Plugins/Online/OnlineFramework/Source/PlayTimeLimit/Public/PlayTimeLimitImpl.h
+++ b/Engine/Plugins/Online/OnlineFramework/Source/PlayTimeLimit/Public/PlayTimeLimitImpl.h
@@ -4,6 +4,7 @@
 
 #include "Interfaces/OnlinePlayTimeLimit.h"
 #include "PlayTimeLimitUser.h"
+#include "Containers/Ticker.h"
 
 /**
  * Configuration
@@ -138,7 +139,7 @@ protected:
 	TArray<FOnlinePlayLimitConfigEntry> ConfigRates;
 
 	/** Delegate for callbacks to Tick */
-	FDelegateHandle TickHandle;
+	FTSTicker::FDelegateHandle TickHandle;
 
 private:
 	// Not copyable
diff --git a/Engine/Plugins/Online/OnlineSubsystem/Source/Private/OnlineSubsystemImpl.cpp b/Engine/Plugins/Online/OnlineSubsystem/Source/Private/OnlineSubsystemImpl.cpp
--- a/Engine/Plugins/Online/OnlineSubsystem/Source/Private/OnlineSubsystemImpl.cpp
+++ b/Engine/Plugins/Online/OnlineSubsystem/Source/Private/OnlineSubsystemImpl.cpp
@@ -31,12 +31,12 @@ namespace OSSConsoleVariables
 const FName FOnlineSubsystemImpl::DefaultInstanceName(TEXT("DefaultInstance"));
 
 FOnlineSubsystemImpl::FOnlineSubsystemImpl(FName InSubsystemName, FName InInstanceName) :
-	FOnlineSubsystemImpl(InSubsystemName, InInstanceName, FTicker::GetCoreTicker())
+	FOnlineSubsystemImpl(InSubsystemName, InInstanceName, FTSTicker::GetCoreTicker())
 {
 }
 
-FOnlineSubsystemImpl::FOnlineSubsystemImpl(FName InSubsystemName, FName InInstanceName, FTicker& Ticker) :
-	FTickerObjectBase(0.0f, Ticker),
+FOnlineSubsystemImpl::FOnlineSubsystemImpl(FName InSubsystemName, FName InInstanceName, FTSTicker& Ticker) :
+	FTSTickerObjectBase(0.0f, Ticker),
 	SubsystemName(InSubsystemName),
 	InstanceName(InInstanceName),
 	bForceDedicated(false),
diff --git a/Engine/Plugins/Online/OnlineSubsystem/Source/Public/OnlineSubsystemImpl.h b/Engine/Plugins/Online/OnlineSubsystem/Source/Public/OnlineSubsystemImpl.h
--- a/Engine/Plugins/Online/OnlineSubsystem/Source/Public/OnlineSubsystemImpl.h
+++ b/Engine/Plugins/Online/OnlineSubsystem/Source/Public/OnlineSubsystemImpl.h
@@ -23,7 +23,7 @@ namespace OSSConsoleVariables
  */
 class ONLINESUBSYSTEM_API FOnlineSubsystemImpl 
 	: public IOnlineSubsystem
-	, public FTickerObjectBase
+	, public FTSTickerObjectBase
 {
 private:
 
@@ -49,7 +49,7 @@ protected:
 	FOnlineSubsystemImpl() = delete;
 
 	FOnlineSubsystemImpl(FName InSubsystemName, FName InInstanceName);
-	FOnlineSubsystemImpl(FName InSubsystemName, FName InInstanceName, FTicker& Ticker);
+	FOnlineSubsystemImpl(FName InSubsystemName, FName InInstanceName, FTSTicker& Ticker);
 
 	/** Name of the subsystem @see OnlineSubsystemNames.h */
 	FName SubsystemName;
@@ -133,7 +133,7 @@ public:
 	virtual void ReloadConfigs(const TSet<FString>& /*ConfigSections*/) override {};
 	virtual FText GetSocialPlatformName() const override;
 
-	// FTickerObjectBase
+	// FTSTickerObjectBase
 	virtual bool Tick(float DeltaTime) override;
 
 	/**
diff --git a/Engine/Plugins/Online/OnlineSubsystemEOS/Source/OnlineSubsystemEOS/Public/OnlineSubsystemEOS.h b/Engine/Plugins/Online/OnlineSubsystemEOS/Source/OnlineSubsystemEOS/Public/OnlineSubsystemEOS.h
--- a/Engine/Plugins/Online/OnlineSubsystemEOS/Source/OnlineSubsystemEOS/Public/OnlineSubsystemEOS.h
+++ b/Engine/Plugins/Online/OnlineSubsystemEOS/Source/OnlineSubsystemEOS/Public/OnlineSubsystemEOS.h
@@ -100,7 +100,7 @@ public:
 	virtual bool Shutdown() override;
 	virtual FString GetAppId() const override;
 
-// FTickerObjectBase
+// FTSTickerObjectBase
 	virtual bool Tick(float DeltaTime) override;
 
 	IVoiceChatUser* GetVoiceChatUserInterface(const FUniqueNetId& LocalUserId);
diff --git a/Engine/Plugins/Online/OnlineSubsystemNull/Source/Private/OnlineMessageSanitizerNull.cpp b/Engine/Plugins/Online/OnlineSubsystemNull/Source/Private/OnlineMessageSanitizerNull.cpp
--- a/Engine/Plugins/Online/OnlineSubsystemNull/Source/Private/OnlineMessageSanitizerNull.cpp
+++ b/Engine/Plugins/Online/OnlineSubsystemNull/Source/Private/OnlineMessageSanitizerNull.cpp
@@ -11,9 +11,9 @@ FMessageSanitizerNull::FMessageSanitizerNull()
 
 FMessageSanitizerNull::~FMessageSanitizerNull()
 {
-	for (TPair<int32, FDelegateHandle>& Pair : RequestIdHandleMap)
+	for (TPair<int32, FTSTicker::FDelegateHandle>& Pair : RequestIdHandleMap)
 	{
-		FTicker::GetCoreTicker().RemoveTicker(Pair.Value);
+		FTSTicker::GetCoreTicker().RemoveTicker(Pair.Value);
 	}
 }
 
@@ -22,7 +22,7 @@ void FMessageSanitizerNull::SanitizeDisplayName(const FString& DisplayName, cons
 	float Delay = 3.0f;
 	int32 ThisRequest = RequestId++;
 
-	FDelegateHandle TickerHandle = FTicker::GetCoreTicker().AddTicker(FTickerDelegate::CreateLambda([this, DisplayName, CompletionDelegate, ThisRequest](float) -> bool
+	FTSTicker::FDelegateHandle TickerHandle = FTSTicker::GetCoreTicker().AddTicker(FTickerDelegate::CreateLambda([this, DisplayName, CompletionDelegate, ThisRequest](float) -> bool
 	{
 		QUICK_SCOPE_CYCLE_COUNTER(STAT_FMessageSanitizerNull_SanitizeDisplayName);
 		FString CleanString;
@@ -41,7 +41,7 @@ void FMessageSanitizerNull::SanitizeDisplayNames(const TArray<FString>& DisplayN
 	float Delay = 3.0f;
 	int32 ThisRequest = RequestId++;
 
-	FDelegateHandle TickerHandle = FTicker::GetCoreTicker().AddTicker(FTickerDelegate::CreateLambda([this, DisplayNames, CompletionDelegate, ThisRequest](float) -> bool
+	FTSTicker::FDelegateHandle TickerHandle = FTSTicker::GetCoreTicker().AddTicker(FTickerDelegate::CreateLambda([this, DisplayNames, CompletionDelegate, ThisRequest](float) -> bool
 	{
 		QUICK_SCOPE_CYCLE_COUNTER(STAT_FMessageSanitizerNull_SanitizeDisplayNames);
 		TArray<FString> OutStrings;
diff --git a/Engine/Plugins/Online/OnlineSubsystemNull/Source/Private/OnlineMessageSanitizerNull.h b/Engine/Plugins/Online/OnlineSubsystemNull/Source/Private/OnlineMessageSanitizerNull.h
--- a/Engine/Plugins/Online/OnlineSubsystemNull/Source/Private/OnlineMessageSanitizerNull.h
+++ b/Engine/Plugins/Online/OnlineSubsystemNull/Source/Private/OnlineMessageSanitizerNull.h
@@ -4,6 +4,7 @@
 
 #include "OnlineSubsystem.h"
 #include "Interfaces/IMessageSanitizerInterface.h"
+#include "Containers/Ticker.h"
 
 /**
  * Implements a degenerate/dummy interface chat message sanitization system for testing
@@ -27,5 +28,5 @@ private:
 	int32 RequestId;
 
 	void PerformSanitize(const FString& RawString, FString& OutCleanString) const;
-	TMap<int32, FDelegateHandle> RequestIdHandleMap;
+	TMap<int32, FTSTicker::FDelegateHandle> RequestIdHandleMap;
 };
\ No newline at end of file
diff --git a/Engine/Plugins/Online/OnlineSubsystemNull/Source/Public/OnlineSubsystemNull.h b/Engine/Plugins/Online/OnlineSubsystemNull/Source/Public/OnlineSubsystemNull.h
--- a/Engine/Plugins/Online/OnlineSubsystemNull/Source/Public/OnlineSubsystemNull.h
+++ b/Engine/Plugins/Online/OnlineSubsystemNull/Source/Public/OnlineSubsystemNull.h
@@ -75,7 +75,7 @@ public:
 	virtual bool Exec(class UWorld* InWorld, const TCHAR* Cmd, FOutputDevice& Ar) override;
 	virtual FText GetOnlineServiceName() const override;
 
-	// FTickerObjectBase
+	// FTSTickerObjectBase
 	
 	virtual bool Tick(float DeltaTime) override;
 
diff --git a/Engine/Plugins/Online/OnlineSubsystemSteam/Source/Private/SocketSubsystemSteam.h b/Engine/Plugins/Online/OnlineSubsystemSteam/Source/Private/SocketSubsystemSteam.h
--- a/Engine/Plugins/Online/OnlineSubsystemSteam/Source/Private/SocketSubsystemSteam.h
+++ b/Engine/Plugins/Online/OnlineSubsystemSteam/Source/Private/SocketSubsystemSteam.h
@@ -16,7 +16,7 @@ class Error;
 /**
  * Windows specific socket subsystem implementation
  */
-class FSocketSubsystemSteam : public ISocketSubsystem, public FTickerObjectBase, public FSelfRegisteringExec
+class FSocketSubsystemSteam : public ISocketSubsystem, public FTSTickerObjectBase, public FSelfRegisteringExec
 {
 protected:
 
diff --git a/Engine/Plugins/Online/OnlineSubsystemTencent/Source/Private/OnlineMessageSanitizerTencent.cpp b/Engine/Plugins/Online/OnlineSubsystemTencent/Source/Private/OnlineMessageSanitizerTencent.cpp
--- a/Engine/Plugins/Online/OnlineSubsystemTencent/Source/Private/OnlineMessageSanitizerTencent.cpp
+++ b/Engine/Plugins/Online/OnlineSubsystemTencent/Source/Private/OnlineMessageSanitizerTencent.cpp
@@ -216,5 +216,5 @@ void FMessageSanitizerTencent::HandleMessageArraySanitized(bool bSuccess,
 FMessageSanitizerTencent::~FMessageSanitizerTencent()
 {
 	ProcessList.Empty();
-	FTicker::GetCoreTicker().RemoveTicker(TickDelegateHandle);
+	FTSTicker::GetCoreTicker().RemoveTicker(TickDelegateHandle);
 }
\ No newline at end of file
diff --git a/Engine/Plugins/Online/OnlineSubsystemTencent/Source/Private/OnlineMessageSanitizerTencent.h b/Engine/Plugins/Online/OnlineSubsystemTencent/Source/Private/OnlineMessageSanitizerTencent.h
--- a/Engine/Plugins/Online/OnlineSubsystemTencent/Source/Private/OnlineMessageSanitizerTencent.h
+++ b/Engine/Plugins/Online/OnlineSubsystemTencent/Source/Private/OnlineMessageSanitizerTencent.h
@@ -150,7 +150,7 @@ public:
 	{
 		check(TencentSubsystem);
 		FTickerDelegate TickDelegate = FTickerDelegate::CreateRaw(this, &FMessageSanitizerTencent::HandleTicker);
-		TickDelegateHandle = FTicker::GetCoreTicker().AddTicker(TickDelegate, 1.0f);
+		TickDelegateHandle = FTSTicker::GetCoreTicker().AddTicker(TickDelegate, 1.0f);
 	}
 
 	virtual ~FMessageSanitizerTencent();
@@ -175,7 +175,7 @@ private:
 	TMap<FString, FString> WordMap;
 
 	/** Handle to the registered TickDelegate. */
-	FDelegateHandle TickDelegateHandle;
+	FTSTicker::FDelegateHandle TickDelegateHandle;
 
 	/** Reference to the main Tencent subsystem */
 	FOnlineSubsystemTencent* TencentSubsystem;
diff --git a/Engine/Plugins/Online/OnlineSubsystemUtils/Source/OnlineSubsystemUtils/Private/Tests/TestCloudInterface.h b/Engine/Plugins/Online/OnlineSubsystemUtils/Source/OnlineSubsystemUtils/Private/Tests/TestCloudInterface.h
--- a/Engine/Plugins/Online/OnlineSubsystemUtils/Source/OnlineSubsystemUtils/Private/Tests/TestCloudInterface.h
+++ b/Engine/Plugins/Online/OnlineSubsystemUtils/Source/OnlineSubsystemUtils/Private/Tests/TestCloudInterface.h
@@ -14,7 +14,7 @@
 /**
  * Class used to test the friends interface
  */
- class FTestCloudInterface : public FTickerObjectBase
+ class FTestCloudInterface : public FTSTickerObjectBase
  {
 	/** The subsystem that was requested to be tested or the default if empty */
 	const FString Subsystem;
@@ -257,7 +257,7 @@
 		SharedCloud = NULL;
 	}
 
-	// FTickerObjectBase
+	// FTSTickerObjectBase
 
 	bool Tick( float DeltaTime ) override;
 
diff --git a/Engine/Plugins/Online/OnlineSubsystemUtils/Source/OnlineSubsystemUtils/Private/Tests/TestLeaderboardInterface.h b/Engine/Plugins/Online/OnlineSubsystemUtils/Source/OnlineSubsystemUtils/Private/Tests/TestLeaderboardInterface.h
--- a/Engine/Plugins/Online/OnlineSubsystemUtils/Source/OnlineSubsystemUtils/Private/Tests/TestLeaderboardInterface.h
+++ b/Engine/Plugins/Online/OnlineSubsystemUtils/Source/OnlineSubsystemUtils/Private/Tests/TestLeaderboardInterface.h
@@ -12,7 +12,7 @@
 /**
  * Class used to test the friends interface
  */
- class FTestLeaderboardInterface : public FTickerObjectBase
+ class FTestLeaderboardInterface : public FTSTickerObjectBase
  {
 	/** The subsystem that was requested to be tested or the default if empty */
 	const FString Subsystem;
@@ -121,7 +121,7 @@
 
 	virtual ~FTestLeaderboardInterface();
 
-	// FTickerObjectBase
+	// FTSTickerObjectBase
 
 	bool Tick( float DeltaTime ) override;
 
diff --git a/Engine/Plugins/Online/OnlineSubsystemUtils/Source/OnlineSubsystemUtils/Private/Tests/TestPresenceInterface.h b/Engine/Plugins/Online/OnlineSubsystemUtils/Source/OnlineSubsystemUtils/Private/Tests/TestPresenceInterface.h
--- a/Engine/Plugins/Online/OnlineSubsystemUtils/Source/OnlineSubsystemUtils/Private/Tests/TestPresenceInterface.h
+++ b/Engine/Plugins/Online/OnlineSubsystemUtils/Source/OnlineSubsystemUtils/Private/Tests/TestPresenceInterface.h
@@ -30,7 +30,7 @@ enum class EPresenceTestStatus : uint8
 
 ENUM_CLASS_FLAGS(EPresenceTestStatus);
 
-class FTestPresenceInterface : public FTickerObjectBase
+class FTestPresenceInterface : public FTSTickerObjectBase
 {
 public:
 	FTestPresenceInterface(const FString& InSubName) :
@@ -48,7 +48,7 @@ public:
 
 	virtual ~FTestPresenceInterface();
 
-	// FTickerObjectBase
+	// FTSTickerObjectBase
 	bool Tick(float DeltaTime) override;
 
 	void Test(class UWorld* InWorld, const FString& RandomUser);
diff --git a/Engine/Plugins/Online/OnlineSubsystemUtils/Source/OnlineSubsystemUtils/Private/Tests/TestSessionInterface.h b/Engine/Plugins/Online/OnlineSubsystemUtils/Source/OnlineSubsystemUtils/Private/Tests/TestSessionInterface.h
--- a/Engine/Plugins/Online/OnlineSubsystemUtils/Source/OnlineSubsystemUtils/Private/Tests/TestSessionInterface.h
+++ b/Engine/Plugins/Online/OnlineSubsystemUtils/Source/OnlineSubsystemUtils/Private/Tests/TestSessionInterface.h
@@ -19,7 +19,7 @@ class TestOnlineSearchSettings;
 /**
  * Class used to test the friends interface
  */
- class FTestSessionInterface : public FTickerObjectBase, public FSelfRegisteringExec
+ class FTestSessionInterface : public FTSTickerObjectBase, public FSelfRegisteringExec
  {
 	/** The subsystem that was requested to be tested or the default if empty */
 	const FString Subsystem;
@@ -342,7 +342,7 @@ class TestOnlineSearchSettings;
 		Friends = NULL;
 	}
 
-	// FTickerObjectBase
+	// FTSTickerObjectBase
 	bool Tick(float DeltaTime) override;
 
 	// FSelfRegisteringExec
diff --git a/Engine/Plugins/Online/OnlineSubsystemUtils/Source/OnlineSubsystemUtils/Private/Tests/TestStatsInterface.h b/Engine/Plugins/Online/OnlineSubsystemUtils/Source/OnlineSubsystemUtils/Private/Tests/TestStatsInterface.h
--- a/Engine/Plugins/Online/OnlineSubsystemUtils/Source/OnlineSubsystemUtils/Private/Tests/TestStatsInterface.h
+++ b/Engine/Plugins/Online/OnlineSubsystemUtils/Source/OnlineSubsystemUtils/Private/Tests/TestStatsInterface.h
@@ -13,7 +13,7 @@
 /**
  * Class used to test the friends interface
  */
- class FTestStatsInterface : public FTickerObjectBase
+ class FTestStatsInterface : public FTSTickerObjectBase
  {
 	/** The subsystem that was requested to be tested or the default if empty */
 	const FString Subsystem;
@@ -69,7 +69,7 @@
 
 	virtual ~FTestStatsInterface();
 
-	// FTickerObjectBase
+	// FTSTickerObjectBase
 
 	bool Tick( float DeltaTime ) override;
 
diff --git a/Engine/Plugins/Online/OnlineSubsystemUtils/Source/OnlineSubsystemUtils/Private/Tests/TestVoice.h b/Engine/Plugins/Online/OnlineSubsystemUtils/Source/OnlineSubsystemUtils/Private/Tests/TestVoice.h
--- a/Engine/Plugins/Online/OnlineSubsystemUtils/Source/OnlineSubsystemUtils/Private/Tests/TestVoice.h
+++ b/Engine/Plugins/Online/OnlineSubsystemUtils/Source/OnlineSubsystemUtils/Private/Tests/TestVoice.h
@@ -18,14 +18,14 @@ enum class EAudioEncodeHint : uint8;
 #define VOICE_MAX_COMPRESSED_BUFFER 20 * 1024
 #define VOICE_STARTING_REMAINDER_SIZE 1 * 1024
 
-class FTestVoice : public FTickerObjectBase, public FSelfRegisteringExec
+class FTestVoice : public FTSTickerObjectBase, public FSelfRegisteringExec
 {
 public:
 
 	FTestVoice();
 	virtual ~FTestVoice();
 
-	// FTickerObjectBase
+	// FTSTickerObjectBase
 	virtual bool Tick(float DeltaTime) override;
 
 	// FSelfRegisteringExec
diff --git a/Engine/Plugins/Performance/PerformanceMonitor/Source/PerformanceMonitor/Private/PerformanceMonitor.cpp b/Engine/Plugins/Performance/PerformanceMonitor/Source/PerformanceMonitor/Private/PerformanceMonitor.cpp
--- a/Engine/Plugins/Performance/PerformanceMonitor/Source/PerformanceMonitor/Private/PerformanceMonitor.cpp
+++ b/Engine/Plugins/Performance/PerformanceMonitor/Source/PerformanceMonitor/Private/PerformanceMonitor.cpp
@@ -107,7 +107,7 @@ bool FPerformanceMonitorModule::Exec(UWorld* InWorld, const TCHAR* Cmd, FOutputD
 				if (!StringCommand.IsEmpty())
 				{
 					StartRecordingPerfTimers(StringCommand, DesiredStats);
-					TickHandler = FTicker::GetCoreTicker().AddTicker(FTickerDelegate::CreateRaw(this, &FPerformanceMonitorModule::Tick));
+					TickHandler = FTSTicker::GetCoreTicker().AddTicker(FTickerDelegate::CreateRaw(this, &FPerformanceMonitorModule::Tick));
 				}
 			}
 			else
@@ -119,7 +119,7 @@ bool FPerformanceMonitorModule::Exec(UWorld* InWorld, const TCHAR* Cmd, FOutputD
 		{
 			if (bRecording)
 			{
-				FTicker::GetCoreTicker().RemoveTicker(TickHandler);
+				FTSTicker::GetCoreTicker().RemoveTicker(TickHandler);
 				StopRecordingPerformanceTimers();
 			}
 			else
diff --git a/Engine/Plugins/Performance/PerformanceMonitor/Source/PerformanceMonitor/Private/PerformanceMonitor.h b/Engine/Plugins/Performance/PerformanceMonitor/Source/PerformanceMonitor/Private/PerformanceMonitor.h
--- a/Engine/Plugins/Performance/PerformanceMonitor/Source/PerformanceMonitor/Private/PerformanceMonitor.h
+++ b/Engine/Plugins/Performance/PerformanceMonitor/Source/PerformanceMonitor/Private/PerformanceMonitor.h
@@ -7,9 +7,7 @@
 #include "Modules/ModuleInterface.h"
 #include "Modules/ModuleManager.h"
 #include "Containers/UnrealString.h"
-/**
- * 
- */
+#include "Containers/Ticker.h"
 
 class FPerformanceMonitorModule : public IModuleInterface, public FSelfRegisteringExec
 {
@@ -31,7 +29,7 @@ private:
 	TMap<FString, TArray<float>> GeneratedStats;
 	FArchive * FileToLogTo;
 	FString LogFileName;
-	FDelegateHandle TickHandler;
+	FTSTicker::FDelegateHandle TickHandler;
 
 	bool bRecording;
 	float TimeOfTestStart;
diff --git a/Engine/Plugins/Runtime/ChunkDownloader/Source/Private/ChunkDownloader.cpp b/Engine/Plugins/Runtime/ChunkDownloader/Source/Private/ChunkDownloader.cpp
--- a/Engine/Plugins/Runtime/ChunkDownloader/Source/Private/ChunkDownloader.cpp
+++ b/Engine/Plugins/Runtime/ChunkDownloader/Source/Private/ChunkDownloader.cpp
@@ -1014,7 +1014,7 @@ void FChunkDownloader::BeginLoadingMode(const FCallback& Callback)
 
 	// compute again next frame (if nothing's queued by then, we'll fire the callback
 	TWeakPtr<FChunkDownloader> WeakThisPtr = AsShared();
-	FTicker::GetCoreTicker().AddTicker(FTickerDelegate::CreateLambda([WeakThisPtr](float dts) {
+	FTSTicker::GetCoreTicker().AddTicker(FTickerDelegate::CreateLambda([WeakThisPtr](float dts) {
 		TSharedPtr<FChunkDownloader> SharedThis = WeakThisPtr.Pin();
 		if (!SharedThis.IsValid() || SharedThis->PostLoadCallbacks.Num() <= 0)
 		{
@@ -1103,7 +1103,7 @@ void FChunkDownloader::ExecuteNextTick(const FCallback& Callback, bool bSuccess)
 {
 	if (Callback)
 	{
-		FTicker::GetCoreTicker().AddTicker(FTickerDelegate::CreateLambda([Callback, bSuccess](float dts) {
+		FTSTicker::GetCoreTicker().AddTicker(FTickerDelegate::CreateLambda([Callback, bSuccess](float dts) {
 			Callback(bSuccess);
 			return false;
 		}));
@@ -1149,7 +1149,7 @@ void FChunkDownloader::TryLoadBuildManifest(int TryNumber)
 		// set a ticker to delay
 		UE_LOG(LogChunkDownloader, Log, TEXT("Will re-attempt manifest download in %f seconds"), SecondsToDelay);
 		TWeakPtr<FChunkDownloader> WeakThisPtr = AsShared();
-		FTicker::GetCoreTicker().AddTicker(FTickerDelegate::CreateLambda([WeakThisPtr, TryNumber](float Unused) {
+		FTSTicker::GetCoreTicker().AddTicker(FTickerDelegate::CreateLambda([WeakThisPtr, TryNumber](float Unused) {
 			TSharedPtr<FChunkDownloader> SharedThis = WeakThisPtr.Pin();
 			if (SharedThis.IsValid())
 			{
@@ -1509,7 +1509,7 @@ void FChunkDownloader::MountChunkInternal(FChunk& Chunk, const FCallback& Callba
 		// start a per-frame ticker until mounts are finished
 		if (!MountTicker.IsValid())
 		{
-			MountTicker = FTicker::GetCoreTicker().AddTicker(FTickerDelegate::CreateSP(this, &FChunkDownloader::UpdateMountTasks));
+			MountTicker = FTSTicker::GetCoreTicker().AddTicker(FTickerDelegate::CreateSP(this, &FChunkDownloader::UpdateMountTasks));
 		}
 	}
 	else
diff --git a/Engine/Plugins/Runtime/ChunkDownloader/Source/Private/Download.cpp b/Engine/Plugins/Runtime/ChunkDownloader/Source/Private/Download.cpp
--- a/Engine/Plugins/Runtime/ChunkDownloader/Source/Private/Download.cpp
+++ b/Engine/Plugins/Runtime/ChunkDownloader/Source/Private/Download.cpp
@@ -30,7 +30,7 @@ void FDownload::Start()
 	if (!HasDeviceSpaceRequired())
 	{
 		TWeakPtr<FDownload> WeakThisPtr = AsShared();
-		FTicker::GetCoreTicker().AddTicker(FTickerDelegate::CreateLambda([WeakThisPtr](float Unused) {
+		FTSTicker::GetCoreTicker().AddTicker(FTickerDelegate::CreateLambda([WeakThisPtr](float Unused) {
 			TSharedPtr<FDownload> SharedThis = WeakThisPtr.Pin();
 			if (SharedThis.IsValid())
 			{
@@ -188,7 +188,7 @@ void FDownload::OnDownloadComplete(const FString& Url, int TryNumber, int32 Http
 	// set a ticker to delay
 	UE_LOG(LogChunkDownloader, Log, TEXT("Will re-attempt to download %s in %f seconds"), *PakFile->Entry.FileName, SecondsToDelay);
 	TWeakPtr<FDownload> WeakThisPtr = AsShared();
-	FTicker::GetCoreTicker().AddTicker(FTickerDelegate::CreateLambda([WeakThisPtr, TryNumber](float Unused) {
+	FTSTicker::GetCoreTicker().AddTicker(FTickerDelegate::CreateLambda([WeakThisPtr, TryNumber](float Unused) {
 		TSharedPtr<FDownload> SharedThis = WeakThisPtr.Pin();
 		if (SharedThis.IsValid() && !SharedThis->bHasCompleted)
 		{
