Game Features: Repaired whitespace from merge
#rnx

diff --git a/Engine/Plugins/Experimental/GameFeatures/Source/GameFeatures/Private/GameFeaturesSubsystem.cpp b/Engine/Plugins/Experimental/GameFeatures/Source/GameFeatures/Private/GameFeaturesSubsystem.cpp
--- a/Engine/Plugins/Experimental/GameFeatures/Source/GameFeatures/Private/GameFeaturesSubsystem.cpp
+++ b/Engine/Plugins/Experimental/GameFeatures/Source/GameFeatures/Private/GameFeaturesSubsystem.cpp
@@ -468,17 +468,17 @@ void UGameFeaturesSubsystem::LoadGameFeaturePlugin(const FString& PluginURL, con
 }
 
 void UGameFeaturesSubsystem::LoadAndActivateGameFeaturePlugin(const FString& PluginURL, const FGameFeaturePluginLoadComplete& CompleteDelegate)
-		{
+{
 	ChangeGameFeatureTargetState(PluginURL, EGameFeatureTargetState::Active, CompleteDelegate);
-		}
+}
 
 void UGameFeaturesSubsystem::ChangeGameFeatureTargetState(const FString& PluginURL, EGameFeatureTargetState TargetState, const FGameFeaturePluginChangeStateComplete& CompleteDelegate)
-		{
+{
 	EGameFeaturePluginState TargetPluginState = EGameFeaturePluginState::MAX;
 	if (TargetState == EGameFeatureTargetState::Active)
 	{
 		TargetPluginState = EGameFeaturePluginState::Active;
-		}
+	}
 	else if (TargetState == EGameFeatureTargetState::Loaded)
 	{
 		TargetPluginState = EGameFeaturePluginState::Loaded;
@@ -490,12 +490,12 @@ void UGameFeaturesSubsystem::ChangeGameFeatureTargetState(const FString& PluginU
 	else if (TargetState == EGameFeatureTargetState::Installed)
 	{
 		TargetPluginState = EGameFeaturePluginState::Installed;
-}
+	}
 
 	if (TargetPluginState != EGameFeaturePluginState::MAX)
-{
-		if (UGameFeaturePluginStateMachine* StateMachine = FindOrCreateGameFeaturePluginStateMachine(PluginURL))
 	{
+		if (UGameFeaturePluginStateMachine* StateMachine = FindOrCreateGameFeaturePluginStateMachine(PluginURL))
+		{
 			if (TargetState == EGameFeatureTargetState::Active
 				&& StateMachine->GetCurrentState() == TargetPluginState)
 			{
@@ -507,7 +507,7 @@ void UGameFeaturesSubsystem::ChangeGameFeatureTargetState(const FString& PluginU
 				for (UObject* Observer : Observers)
 				{
 					CastChecked<IGameFeatureStateChangeObserver>(Observer)->OnGameFeatureActivating(StateMachine->GetGameFeatureDataForActivePlugin());
-	}
+				}
 
 				FTSTicker::GetCoreTicker().AddTicker(FTickerDelegate::CreateWeakLambda(this, [CompleteDelegate](float)
 					{
@@ -524,8 +524,8 @@ void UGameFeaturesSubsystem::ChangeGameFeatureTargetState(const FString& PluginU
 						return false;
 					}));
 			}
-	else
-	{
+			else
+			{
 				StateMachine->SetDestinationState(TargetPluginState, FGameFeatureStateTransitionComplete::CreateUObject(this, &ThisClass::ChangeGameFeatureTargetStateComplete, CompleteDelegate));
 			}
 		}
