GameFeatures - Adding FilterInactivePluginAssets, which allows users to easily remove content from a bulk set of content which shouldn't be examined because it's from an inactive plugin.
GameFeatures - General cleanup, making several functions const that should have been const.
CommonConversation - Now ignores conversation graphs from inactive game feature plugins.

Michael.Noland, Bob.Tellez

diff --git a/Engine/Plugins/Experimental/CommonConversation/CommonConversation.uplugin b/Engine/Plugins/Experimental/CommonConversation/CommonConversation.uplugin
--- a/Engine/Plugins/Experimental/CommonConversation/CommonConversation.uplugin
+++ b/Engine/Plugins/Experimental/CommonConversation/CommonConversation.uplugin
@@ -32,5 +32,9 @@
 		}
 	],
 	"Plugins": [
+		{
+			"Name": "GameFeatures",
+			"Enabled": true
+		}
 	]
 }
\ No newline at end of file
diff --git a/Engine/Plugins/Experimental/CommonConversation/Source/CommonConversationRuntime/CommonConversationRuntime.Build.cs b/Engine/Plugins/Experimental/CommonConversation/Source/CommonConversationRuntime/CommonConversationRuntime.Build.cs
--- a/Engine/Plugins/Experimental/CommonConversation/Source/CommonConversationRuntime/CommonConversationRuntime.Build.cs
+++ b/Engine/Plugins/Experimental/CommonConversation/Source/CommonConversationRuntime/CommonConversationRuntime.Build.cs
@@ -41,6 +41,7 @@ public class CommonConversationRuntime : ModuleRules
 				"SlateCore",
 				"Projects",
 				"AIModule",
+				"GameFeatures"
 				// ... add private dependencies that you statically link with here ...	
 			}
 		);
diff --git a/Engine/Plugins/Experimental/CommonConversation/Source/CommonConversationRuntime/Private/ConversationRegistry.cpp b/Engine/Plugins/Experimental/CommonConversation/Source/CommonConversationRuntime/Private/ConversationRegistry.cpp
--- a/Engine/Plugins/Experimental/CommonConversation/Source/CommonConversationRuntime/Private/ConversationRegistry.cpp
+++ b/Engine/Plugins/Experimental/CommonConversation/Source/CommonConversationRuntime/Private/ConversationRegistry.cpp
@@ -11,6 +11,7 @@
 #include "Engine/World.h"
 #include "ConversationContext.h"
 #include "Engine/StreamableManager.h"
+#include "GameFeaturesSubsystem.h"
 
 //======================================================================================
 
@@ -299,6 +300,9 @@ void UConversationRegistry::BuildDependenciesGraph()
 	TArray<FAssetData> AllConversations;
 	UAssetManager::Get().GetPrimaryAssetDataList(FPrimaryAssetType(UConversationDatabase::StaticClass()->GetFName()), AllConversations);
 
+	// Don't index conversation graphs from inactive game feature plugins.
+	UGameFeaturesSubsystem::Get().FilterInactivePluginAssets(AllConversations);
+
 	RuntimeDependencyGraph.Reset();
 	EntryTagToConversations.Reset();
 	NodeGuidToConversation.Reset();
diff --git a/Engine/Plugins/Experimental/GameFeatures/Source/GameFeatures/Private/GameFeaturePluginStateMachine.cpp b/Engine/Plugins/Experimental/GameFeatures/Source/GameFeatures/Private/GameFeaturePluginStateMachine.cpp
--- a/Engine/Plugins/Experimental/GameFeatures/Source/GameFeatures/Private/GameFeaturePluginStateMachine.cpp
+++ b/Engine/Plugins/Experimental/GameFeatures/Source/GameFeatures/Private/GameFeaturePluginStateMachine.cpp
@@ -1253,7 +1253,8 @@ struct FGameFeaturePluginState_Deactivating : public FGameFeaturePluginState
 
 			// Deactivate
 			FGameFeatureDeactivatingContext Context(FSimpleDelegate::CreateRaw(this, &FGameFeaturePluginState_Deactivating::OnPauserCompleted));
-			UGameFeaturesSubsystem::Get().OnGameFeatureDeactivating(StateProperties.GameFeatureData, Context);
+			const FString PluginName = FPaths::GetBaseFilename(StateProperties.PluginInstalledFilename);
+			UGameFeaturesSubsystem::Get().OnGameFeatureDeactivating(StateProperties.GameFeatureData, PluginName, Context);
 			NumExpectedPausers = Context.NumPausers;
 		}
 
@@ -1281,7 +1282,8 @@ struct FGameFeaturePluginState_Activating : public FGameFeaturePluginState
 
 		StateProperties.GameFeatureData->InitializeHierarchicalPluginIniFiles(StateProperties.PluginInstalledFilename);
 
-		UGameFeaturesSubsystem::Get().OnGameFeatureActivating(StateProperties.GameFeatureData);
+		const FString PluginName = FPaths::GetBaseFilename(StateProperties.PluginInstalledFilename);
+		UGameFeaturesSubsystem::Get().OnGameFeatureActivating(StateProperties.GameFeatureData, PluginName);
 
 		StateStatus.SetTransition(EGameFeaturePluginState::Active);
 	}
diff --git a/Engine/Plugins/Experimental/GameFeatures/Source/GameFeatures/Private/GameFeaturesSubsystem.cpp b/Engine/Plugins/Experimental/GameFeatures/Source/GameFeatures/Private/GameFeaturesSubsystem.cpp
--- a/Engine/Plugins/Experimental/GameFeatures/Source/GameFeatures/Private/GameFeaturesSubsystem.cpp
+++ b/Engine/Plugins/Experimental/GameFeatures/Source/GameFeatures/Private/GameFeaturesSubsystem.cpp
@@ -301,6 +301,8 @@ void UGameFeaturesSubsystem::OnGameFeatureRegistering(const UGameFeatureData* Ga
 	check(GameFeatureData);
 	AddGameFeatureToAssetManager(GameFeatureData, PluginName);
 
+	InactivePluginNames.Add(PluginName);
+
 	for (UGameFeatureStateChangeObserver* Observer : Observers)
 	{
 		Observer->OnGameFeatureRegistering(GameFeatureData, PluginName);
@@ -332,9 +334,11 @@ void UGameFeaturesSubsystem::OnGameFeatureLoading(const UGameFeatureData* GameFe
 	}
 }
 
-void UGameFeaturesSubsystem::OnGameFeatureActivating(const UGameFeatureData* GameFeatureData)
+void UGameFeaturesSubsystem::OnGameFeatureActivating(const UGameFeatureData* GameFeatureData, const FString& PluginName)
 {
 	check(GameFeatureData);
+	
+	InactivePluginNames.Remove(PluginName);
 
 	for (UGameFeatureStateChangeObserver* Observer : Observers)
 	{
@@ -350,9 +354,12 @@ void UGameFeaturesSubsystem::OnGameFeatureActivating(const UGameFeatureData* Gam
 	}
 }
 
-void UGameFeaturesSubsystem::OnGameFeatureDeactivating(const UGameFeatureData* GameFeatureData, FGameFeatureDeactivatingContext& Context)
+void UGameFeaturesSubsystem::OnGameFeatureDeactivating(const UGameFeatureData* GameFeatureData, const FString& PluginName, FGameFeatureDeactivatingContext& Context)
 {
 	check(GameFeatureData);
+
+	InactivePluginNames.Add(PluginName);
+
 	for (UGameFeatureStateChangeObserver* Observer : Observers)
 	{
 		Observer->OnGameFeatureDeactivating(GameFeatureData, Context);
@@ -390,7 +397,7 @@ void UGameFeaturesSubsystem::GetGameFeatureDataForActivePlugins(TArray<const UGa
 
 const UGameFeatureData* UGameFeaturesSubsystem::GetGameFeatureDataForActivePluginByURL(const FString& PluginURL)
 {
-	if (UGameFeaturePluginStateMachine* GFSM = GetGameFeaturePluginStateMachine(PluginURL, false))
+	if (UGameFeaturePluginStateMachine* GFSM = FindGameFeaturePluginStateMachine(PluginURL))
 	{
 		return GFSM->GetGameFeatureDataForActivePlugin();
 	}
@@ -423,7 +430,7 @@ void UGameFeaturesSubsystem::LoadAndActivateGameFeaturePlugin(const FString& Plu
 {
 	if (GameSpecificPolicies->IsPluginAllowed(PluginURL))
 	{
-		UGameFeaturePluginStateMachine* StateMachine = GetGameFeaturePluginStateMachine(PluginURL, true);
+		UGameFeaturePluginStateMachine* StateMachine = FindOrCreateGameFeaturePluginStateMachine(PluginURL);
 		StateMachine->SetDestinationState(EGameFeaturePluginState::Active, FGameFeatureStateTransitionComplete::CreateUObject(this, &ThisClass::LoadExternallyRequestedGameFeaturePluginComplete, CompleteDelegate));
 	}
 	else
@@ -432,9 +439,9 @@ void UGameFeaturesSubsystem::LoadAndActivateGameFeaturePlugin(const FString& Plu
 	}
 }
 
-bool UGameFeaturesSubsystem::GetGameFeaturePluginInstallPercent(const FString& PluginURL, float& Install_Percent)
+bool UGameFeaturesSubsystem::GetGameFeaturePluginInstallPercent(const FString& PluginURL, float& Install_Percent) const
 {
-	if (UGameFeaturePluginStateMachine* StateMachine = GetGameFeaturePluginStateMachine(PluginURL, false))
+	if (const UGameFeaturePluginStateMachine* StateMachine = FindGameFeaturePluginStateMachine(PluginURL))
 	{
 		if (StateMachine->IsStatusKnown() && StateMachine->IsAvailable())
 		{
@@ -456,9 +463,9 @@ bool UGameFeaturesSubsystem::GetGameFeaturePluginInstallPercent(const FString& P
 	return false;
 }
 
-bool UGameFeaturesSubsystem::IsGameFeaturePluginActive(const FString& PluginURL)
+bool UGameFeaturesSubsystem::IsGameFeaturePluginActive(const FString& PluginURL) const
 {
-	if (UGameFeaturePluginStateMachine* StateMachine = GetGameFeaturePluginStateMachine(PluginURL, false))
+	if (const UGameFeaturePluginStateMachine* StateMachine = FindGameFeaturePluginStateMachine(PluginURL))
 	{
 		return StateMachine->GetCurrentState() == EGameFeaturePluginState::Active;
 	}
@@ -474,7 +481,7 @@ void UGameFeaturesSubsystem::DeactivateGameFeaturePlugin(const FString& PluginUR
 
 void UGameFeaturesSubsystem::DeactivateGameFeaturePlugin(const FString& PluginURL, const FGameFeaturePluginDeactivateComplete& CompleteDelegate)
 {
-	if (UGameFeaturePluginStateMachine* StateMachine = GetGameFeaturePluginStateMachine(PluginURL, false))
+	if (UGameFeaturePluginStateMachine* StateMachine = FindGameFeaturePluginStateMachine(PluginURL))
 	{
 		ensureAlwaysMsgf(StateMachine->GetCurrentState() == StateMachine->GetDestinationState(), TEXT("Setting a new destination state while state machine is running!"));
 
@@ -509,7 +516,7 @@ void UGameFeaturesSubsystem::UnloadGameFeaturePlugin(const FString& PluginURL, b
 
 void UGameFeaturesSubsystem::UnloadGameFeaturePlugin(const FString& PluginURL, const FGameFeaturePluginUnloadComplete& CompleteDelegate, bool bKeepRegistered)
 {
-	if (UGameFeaturePluginStateMachine* StateMachine = GetGameFeaturePluginStateMachine(PluginURL, false))
+	if (UGameFeaturePluginStateMachine* StateMachine = FindGameFeaturePluginStateMachine(PluginURL))
 	{
 		ensureAlwaysMsgf(StateMachine->GetCurrentState() == StateMachine->GetDestinationState(), TEXT("Setting a new destination state while state machine is running!"));
 
@@ -545,7 +552,7 @@ void UGameFeaturesSubsystem::UninstallGameFeaturePlugin(const FString& PluginURL
 
 void UGameFeaturesSubsystem::UninstallGameFeaturePlugin(const FString& PluginURL, const FGameFeaturePluginUninstallComplete& CompleteDelegate)
 {
-	if (UGameFeaturePluginStateMachine* StateMachine = GetGameFeaturePluginStateMachine(PluginURL, false))
+	if (UGameFeaturePluginStateMachine* StateMachine = FindGameFeaturePluginStateMachine(PluginURL))
 	{
 		ensureAlwaysMsgf(StateMachine->GetCurrentState() == StateMachine->GetDestinationState(), TEXT("Setting a new destination state while state machine is running!"));
 
@@ -590,7 +597,7 @@ void UGameFeaturesSubsystem::LoadBuiltInGameFeaturePlugin(const TSharedRef<IPlug
 				bool bShouldProcess = AdditionalFilter(PluginDescriptorFilename, PluginDetails, BehaviorOptions);
 				if (bShouldProcess)
 				{
-					UGameFeaturePluginStateMachine* StateMachine = GetGameFeaturePluginStateMachine(PluginURL, true);
+					UGameFeaturePluginStateMachine* StateMachine = FindOrCreateGameFeaturePluginStateMachine(PluginURL);
 
 					const EBuiltInAutoState InitialAutoState = (BehaviorOptions.AutoStateOverride != EBuiltInAutoState::Invalid) ? BehaviorOptions.AutoStateOverride : PluginDetails.BuiltInAutoState;
 						
@@ -629,9 +636,9 @@ void UGameFeaturesSubsystem::LoadBuiltInGameFeaturePlugins(FBuiltInPluginAdditio
 	UAssetManager::Get().PopBulkScanning();
 }
 
-bool UGameFeaturesSubsystem::GetPluginURLForBuiltInPluginByName(const FString& PluginName, FString& OutPluginURL)
+bool UGameFeaturesSubsystem::GetPluginURLForBuiltInPluginByName(const FString& PluginName, FString& OutPluginURL) const
 {
-	if (FString* PluginURL = GameFeaturePluginNameToPathMap.Find(PluginName))
+	if (const FString* PluginURL = GameFeaturePluginNameToPathMap.Find(PluginName))
 	{
 		OutPluginURL = *PluginURL;
 		return true;
@@ -640,10 +647,10 @@ bool UGameFeaturesSubsystem::GetPluginURLForBuiltInPluginByName(const FString& P
 	return false;
 }
 
-FString UGameFeaturesSubsystem::GetPluginFilenameFromPluginURL(const FString& PluginURL)
+FString UGameFeaturesSubsystem::GetPluginFilenameFromPluginURL(const FString& PluginURL) const
 {
 	FString PluginFilename;
-	if (UGameFeaturePluginStateMachine* GFSM = GetGameFeaturePluginStateMachine(PluginURL, false))
+	if (const UGameFeaturePluginStateMachine* GFSM = FindGameFeaturePluginStateMachine(PluginURL))
 	{
 		GFSM->GetPluginFilename(PluginFilename);
 	}
@@ -788,17 +795,38 @@ bool UGameFeaturesSubsystem::GetGameFeaturePluginDetails(const FString& PluginDe
 	return true;
 }
 
-UGameFeaturePluginStateMachine* UGameFeaturesSubsystem::GetGameFeaturePluginStateMachine(const FString& PluginURL, bool bCreateIfItDoesntExist)
+UGameFeaturePluginStateMachine* UGameFeaturesSubsystem::FindGameFeaturePluginStateMachineByPluginName(const FString& PluginName) const
 {
-	UGameFeaturePluginStateMachine** ExistingStateMachine = GameFeaturePluginStateMachines.Find(PluginURL);
+	for (auto StateMachineIt = GameFeaturePluginStateMachines.CreateConstIterator(); StateMachineIt; ++StateMachineIt)
+	{
+		if (UGameFeaturePluginStateMachine* GFSM = StateMachineIt.Value())
+		{
+			if (GFSM->GetGameFeatureName() == PluginName)
+			{
+				return GFSM;
+			}
+		}
+	}
+
+	return nullptr;
+}
+
+UGameFeaturePluginStateMachine* UGameFeaturesSubsystem::FindGameFeaturePluginStateMachine(const FString& PluginURL) const
+{
+	UGameFeaturePluginStateMachine* const* ExistingStateMachine = GameFeaturePluginStateMachines.Find(PluginURL);
 	if (ExistingStateMachine)
 	{
 		return *ExistingStateMachine;
 	}
 
-	if (!bCreateIfItDoesntExist)
+	return nullptr;
+}
+
+UGameFeaturePluginStateMachine* UGameFeaturesSubsystem::FindOrCreateGameFeaturePluginStateMachine(const FString& PluginURL)
+{
+	if (UGameFeaturePluginStateMachine* ExistingStateMachine = FindGameFeaturePluginStateMachine(PluginURL))
 	{
-		return nullptr;
+		return ExistingStateMachine;
 	}
 
 	UGameFeaturePluginStateMachine* NewStateMachine = NewObject<UGameFeaturePluginStateMachine>(this);
@@ -868,7 +896,7 @@ bool UGameFeaturesSubsystem::HandleRequestPluginDependencyStateMachines(const FS
 	{
 		for (const FString& DependencyURL : Details.PluginDependencies)
 		{
-			UGameFeaturePluginStateMachine* Dependency = GetGameFeaturePluginStateMachine(DependencyURL, true);
+			UGameFeaturePluginStateMachine* Dependency = FindOrCreateGameFeaturePluginStateMachine(DependencyURL);
 			check(Dependency);
 			OutDependencyMachines.Add(Dependency);
 		}
@@ -928,6 +956,35 @@ void UGameFeaturesSubsystem::ListGameFeaturePlugins(const TArray<FString>& Args,
 	Ar.Logf(TEXT("Total Game Feature Plugins: %d"), PluginCount);
 }
 
+bool UGameFeaturesSubsystem::IsContentWithinInactivePlugin(const FString& InObjectOrPackagePath) const
+{
+	// Look for the first slash beyond the first one we start with.
+	const int32 RootEndIndex = InObjectOrPackagePath.Find(TEXT("/"), ESearchCase::IgnoreCase, ESearchDir::FromStart, 1);
+
+	const FString ObjectPathRootName = InObjectOrPackagePath.Mid(1, RootEndIndex - 1);
+
+	if (InactivePluginNames.Contains(ObjectPathRootName))
+	{
+		return true;
+	}
+	
+	return false;
+}
+
+void UGameFeaturesSubsystem::FilterInactivePluginAssets(TArray<FAssetIdentifier>& AssetsToFilter) const
+{
+	AssetsToFilter.RemoveAllSwap([&](const FAssetIdentifier& Asset) {
+		return IsContentWithinInactivePlugin(Asset.PackageName.ToString());
+	});
+}
+
+void UGameFeaturesSubsystem::FilterInactivePluginAssets(TArray<FAssetData>& AssetsToFilter) const
+{
+	AssetsToFilter.RemoveAllSwap([&](const FAssetData& Asset) {
+		return IsContentWithinInactivePlugin(Asset.ObjectPath.ToString());
+	});
+}
+
 EBuiltInAutoState UGameFeaturesSubsystem::DetermineBuiltInInitialFeatureState(TSharedPtr<FJsonObject> Descriptor, const FString& ErrorContext)
 {
 	EBuiltInAutoState InitialState = EBuiltInAutoState::Invalid;
diff --git a/Engine/Plugins/Experimental/GameFeatures/Source/GameFeatures/Public/GameFeaturesSubsystem.h b/Engine/Plugins/Experimental/GameFeatures/Source/GameFeatures/Public/GameFeaturesSubsystem.h
--- a/Engine/Plugins/Experimental/GameFeatures/Source/GameFeatures/Public/GameFeaturesSubsystem.h
+++ b/Engine/Plugins/Experimental/GameFeatures/Source/GameFeatures/Public/GameFeaturesSubsystem.h
@@ -12,6 +12,7 @@
 class UGameFeaturePluginStateMachine;
 class UGameFeatureStateChangeObserver;
 struct FStreamableHandle;
+struct FAssetIdentifier;
 class UGameFeatureData;
 class UGameFeaturesProjectPolicies;
 enum class EGameFeaturePluginState : uint8;
@@ -153,10 +154,10 @@ public:
 	void LoadAndActivateGameFeaturePlugin(const FString& PluginURL, const FGameFeaturePluginLoadComplete& CompleteDelegate);
 
 	/** Gets the Install_Percent for single game feature plugin if it is active. */
-	bool GetGameFeaturePluginInstallPercent(const FString& PluginURL, float& Install_Percent);
+	bool GetGameFeaturePluginInstallPercent(const FString& PluginURL, float& Install_Percent) const;
 
 	/** Determines if a plugin is in the Active state.*/
-	bool IsGameFeaturePluginActive(const FString& PluginURL);
+	bool IsGameFeaturePluginActive(const FString& PluginURL) const;
 
 	/** Deactivates the specified plugin */
 	void DeactivateGameFeaturePlugin(const FString& PluginURL);
@@ -171,10 +172,13 @@ public:
 	void UninstallGameFeaturePlugin(const FString& PluginURL, const FGameFeaturePluginUninstallComplete& CompleteDelegate);
 
 	/** If the specified plugin is a built-in plugin, return the URL used to identify it. Returns true if the plugin exists, false if it was not found */
-	bool GetPluginURLForBuiltInPluginByName(const FString& PluginName, FString& OutPluginURL);
+	bool GetPluginURLForBuiltInPluginByName(const FString& PluginName, FString& OutPluginURL) const;
 	
+	/** Get the plugin path from the plugin name */
+	FString GetPluginFilenameFromPluginName(const FString& PluginName);
+
 	/** Get the plugin path from the plugin URL */
-	FString GetPluginFilenameFromPluginURL(const FString& PluginURL);
+	FString GetPluginFilenameFromPluginURL(const FString& PluginURL) const;
 
 	/** Fixes a package path/directory to either be relative to plugin root or not. Paths relative to different roots will not be modified */
 	static void FixPluginPackagePath(FString& PathToFix, const FString& PluginRootPath, bool bMakeRelativeToPluginRoot);
@@ -198,7 +202,16 @@ public:
 	//@TODO: GameFeaturePluginEnginePush: Might not be general enough for engine level, TBD
 	void GetLoadedGameFeaturePluginFilenamesForCooking(TArray<FString>& OutLoadedPluginFilenames) const;
 	
-	/** Broadcasts when a plugin is activated and the GameFeatureData is available */
+	/** Is content inside inactive plugin? */
+	bool IsContentWithinInactivePlugin(const FString& ObjectOrPackagePath) const;
+
+	/** Removes assets that are in plugins we know to be inactive.  Order is not maintained. */
+	void FilterInactivePluginAssets(TArray<FAssetIdentifier>& AssetsToFilter) const;
+
+	/** Removes assets that are in plugins we know to be inactive.  Order is not maintained. */
+	void FilterInactivePluginAssets(TArray<FAssetData>& AssetsToFilter) const;
+
+    /** Broadcasts when a plugin is activated and the GameFeatureData is available */
 	static FGameFeaturePluginLoadCompleteDataReady& OnPluginLoadCompleteDataReady() { return PluginLoadedGameFeatureDataReadyDelegate; }
 
 	/** Broadcasts when a plugin is deactivated */
@@ -216,10 +229,10 @@ private:
 	void OnGameFeatureRegistering(const UGameFeatureData* GameFeatureData, const FString& PluginName);
 	friend struct FGameFeaturePluginState_Registering;
 
-	void OnGameFeatureActivating(const UGameFeatureData* GameFeatureData);
+	void OnGameFeatureActivating(const UGameFeatureData* GameFeatureData, const FString& PluginName);
 	friend struct FGameFeaturePluginState_Activating;
 
-	void OnGameFeatureDeactivating(const UGameFeatureData* GameFeatureData, FGameFeatureDeactivatingContext& Context);
+	void OnGameFeatureDeactivating(const UGameFeatureData* GameFeatureData, const FString& PluginName, FGameFeatureDeactivatingContext& Context);
 	friend struct FGameFeaturePluginState_Deactivating;
 
 	void OnGameFeatureLoading(const UGameFeatureData* GameFeatureData);
@@ -238,8 +251,14 @@ private:
 	/** Gets relevant properties out of a uplugin file */
 	bool GetGameFeaturePluginDetails(const FString& PluginDescriptorFilename, struct FGameFeaturePluginDetails& OutPluginDetails) const;
 
-	/** Gets the state machine associated with the specified URL, creates it if it doesnt exist when bCreateIfItDoesntExist is true */
-	UGameFeaturePluginStateMachine* GetGameFeaturePluginStateMachine(const FString& PluginURL, bool bCreateIfItDoesntExist);
+	/** Gets the state machine associated with the specified plugin name */
+	UGameFeaturePluginStateMachine* FindGameFeaturePluginStateMachineByPluginName(const FString& PluginName) const;
+
+	/** Gets the state machine associated with the specified URL */
+	UGameFeaturePluginStateMachine* FindGameFeaturePluginStateMachine(const FString& PluginURL) const;
+
+	/** Gets the state machine associated with the specified URL, creates it if it doesnt exist */
+	UGameFeaturePluginStateMachine* FindOrCreateGameFeaturePluginStateMachine(const FString& PluginURL);
 
 	/** Notification that a game feature has finished loading, and whether it was successful */
 	void LoadGameFeaturePluginComplete(UGameFeaturePluginStateMachine* Machine, const UE::GameFeatures::FResult& Result);
@@ -267,9 +286,11 @@ private:
 	UPROPERTY(Transient)
 	TMap<FString, UGameFeaturePluginStateMachine*> GameFeaturePluginStateMachines;
 
-
 	TMap<FString, FString> GameFeaturePluginNameToPathMap;
 
+	/** The set of plugins that are currently registered, but inactive. */
+	TSet<FString> InactivePluginNames;
+
 	UPROPERTY(Transient)
 	TArray<UGameFeatureStateChangeObserver*> Observers;
 
