Fix build issues after API changes from merge.

diff --git a/Engine/Plugins/Experimental/GameFeatures/Source/GameFeatures/Private/GameFeaturesSubsystem.cpp b/Engine/Plugins/Experimental/GameFeatures/Source/GameFeatures/Private/GameFeaturesSubsystem.cpp
--- a/Engine/Plugins/Experimental/GameFeatures/Source/GameFeatures/Private/GameFeaturesSubsystem.cpp
+++ b/Engine/Plugins/Experimental/GameFeatures/Source/GameFeatures/Private/GameFeaturesSubsystem.cpp
@@ -409,7 +409,7 @@ void UGameFeaturesSubsystem::LoadGameFeaturePlugin(const FString& PluginURL, con
 {
 	if (GameSpecificPolicies->IsPluginAllowed(PluginURL))
 	{
-		UGameFeaturePluginStateMachine* StateMachine = GetGameFeaturePluginStateMachine(PluginURL, true);
+		UGameFeaturePluginStateMachine* StateMachine = FindOrCreateGameFeaturePluginStateMachine(PluginURL);
 
 		if (StateMachine->GetCurrentState() < EGameFeaturePluginState::Loaded)
 		{
@@ -695,7 +695,7 @@ void UGameFeaturesSubsystem::GetLoadedGameFeaturePluginFilenamesForCooking(TArra
 
 EGameFeaturePluginState UGameFeaturesSubsystem::GetPluginState(const FString& PluginURL)
 {
-	if (UGameFeaturePluginStateMachine* StateMachine = GetGameFeaturePluginStateMachine(PluginURL, /*bCreateIfDoesntExist=*/ false))
+	if (UGameFeaturePluginStateMachine* StateMachine = FindGameFeaturePluginStateMachine(PluginURL))
 	{
 		return StateMachine->GetCurrentState();
 	}
