Back out changelist 16312429. CIS errors

diff --git a/Engine/Plugins/Experimental/GameFeatures/Source/GameFeatures/Private/GameFeaturePluginStateMachine.cpp b/Engine/Plugins/Experimental/GameFeatures/Source/GameFeatures/Private/GameFeaturePluginStateMachine.cpp
--- a/Engine/Plugins/Experimental/GameFeatures/Source/GameFeatures/Private/GameFeaturePluginStateMachine.cpp
+++ b/Engine/Plugins/Experimental/GameFeatures/Source/GameFeatures/Private/GameFeaturePluginStateMachine.cpp
@@ -979,15 +979,15 @@ struct FGameFeaturePluginState_WaitingForDependencies : public FGameFeaturePlugi
 				for (UGameFeaturePluginStateMachine* Dependency : Dependencies)
 				{
 					check(Dependency);
-					if (Dependency->GetCurrentState() < EGameFeaturePluginState::Registered)
+					if (Dependency->GetCurrentState() < EGameFeaturePluginState::Loaded)
 					{
 						RemainingDependencies.Add(Dependency);
 						Dependency->OnStateChanged().AddRaw(this, &FGameFeaturePluginState_WaitingForDependencies::OnDependencyStateChanged);
 
 						// If we are not alreadying loading this dependency, do so now
-						if (Dependency->GetDestinationState() < EGameFeaturePluginState::Registered)
+						if (Dependency->GetDestinationState() < EGameFeaturePluginState::Loaded)
 						{
-							Dependency->SetDestinationState(EGameFeaturePluginState::Registered, FGameFeatureStateTransitionComplete());
+							Dependency->SetDestinationState(EGameFeaturePluginState::Loaded, FGameFeatureStateTransitionComplete());
 						}
 					}
 				}
@@ -1009,7 +1009,7 @@ struct FGameFeaturePluginState_WaitingForDependencies : public FGameFeaturePlugi
 				StateStatus.SetTransitionError(EGameFeaturePluginState::ErrorWaitingForDependencies, UE::GameFeatures::StateMachineErrorNamespace + TEXT("Dependency_Destroyed_Before_Finish"));
 				return;
 			}
-			else if (RemainingDependency->GetCurrentState() >= EGameFeaturePluginState::Registered)
+			else if (RemainingDependency->GetCurrentState() >= EGameFeaturePluginState::Loaded)
 			{
 				RemainingDependency->OnStateChanged().RemoveAll(this);
 				RemainingDependencies.RemoveAt(DepIdx, 1, false);
