Initial load time optimization

Reduces Engine::Init() from 5.69s -> 4.94s on a Gen5 console for an internal project together with another UpdateTemporaryCaches() optimization

diff --git a/Engine/Plugins/Experimental/GameFeatures/Source/GameFeatures/Private/GameFeaturesSubsystem.cpp b/Engine/Plugins/Experimental/GameFeatures/Source/GameFeatures/Private/GameFeaturesSubsystem.cpp
--- a/Engine/Plugins/Experimental/GameFeatures/Source/GameFeatures/Private/GameFeaturesSubsystem.cpp
+++ b/Engine/Plugins/Experimental/GameFeatures/Source/GameFeatures/Private/GameFeaturesSubsystem.cpp
@@ -581,6 +581,8 @@ void UGameFeaturesSubsystem::UninstallGameFeaturePlugin(const FString& PluginURL
 
 void UGameFeaturesSubsystem::LoadBuiltInGameFeaturePlugin(const TSharedRef<IPlugin>& Plugin, FBuiltInPluginAdditionalFilters AdditionalFilter)
 {
+	UAssetManager::Get().PushBulkScanning();
+
 	const FString& PluginDescriptorFilename = Plugin->GetDescriptorFileName();
 
 	// Make sure you are in the game feature plugins folder. All GameFeaturePlugins are in this folder.
