#UE Avoid an ensure when you ask to Deactivate/Unload/Uninstall a plugin that is already in that state. This can only really be avoided when you are not supplying a callback, since the underlying code will need to keep track of callbacks and fire them all to fully fix this issue.

diff --git a/Engine/Plugins/Experimental/GameFeatures/Source/GameFeatures/Private/GameFeaturesSubsystem.cpp b/Engine/Plugins/Experimental/GameFeatures/Source/GameFeatures/Private/GameFeaturesSubsystem.cpp
--- a/Engine/Plugins/Experimental/GameFeatures/Source/GameFeatures/Private/GameFeaturesSubsystem.cpp
+++ b/Engine/Plugins/Experimental/GameFeatures/Source/GameFeatures/Private/GameFeaturesSubsystem.cpp
@@ -475,8 +475,14 @@ bool UGameFeaturesSubsystem::IsGameFeaturePluginActive(const FString& PluginURL)
 
 void UGameFeaturesSubsystem::DeactivateGameFeaturePlugin(const FString& PluginURL)
 {
-	FGameFeaturePluginDeactivateComplete Callback = FGameFeaturePluginDeactivateComplete();
-	DeactivateGameFeaturePlugin(PluginURL, Callback);
+	if (UGameFeaturePluginStateMachine* StateMachine = FindGameFeaturePluginStateMachine(PluginURL))
+	{
+		if (StateMachine->GetDestinationState() > EGameFeaturePluginState::Loaded)
+		{
+			FGameFeaturePluginDeactivateComplete Callback = FGameFeaturePluginDeactivateComplete();
+			DeactivateGameFeaturePlugin(PluginURL, Callback);
+		}
+	}
 }
 
 void UGameFeaturesSubsystem::DeactivateGameFeaturePlugin(const FString& PluginURL, const FGameFeaturePluginDeactivateComplete& CompleteDelegate)
@@ -510,8 +516,15 @@ void UGameFeaturesSubsystem::DeactivateGameFeaturePlugin(const FString& PluginUR
 
 void UGameFeaturesSubsystem::UnloadGameFeaturePlugin(const FString& PluginURL, bool bKeepRegistered)
 {
-	FGameFeaturePluginUnloadComplete Callback = FGameFeaturePluginUnloadComplete();
-	UnloadGameFeaturePlugin(PluginURL, Callback, bKeepRegistered);
+	if (UGameFeaturePluginStateMachine* StateMachine = FindGameFeaturePluginStateMachine(PluginURL))
+	{
+		EGameFeaturePluginState::Type DestinationState = bKeepRegistered ? EGameFeaturePluginState::Registered : EGameFeaturePluginState::Installed;
+		if (StateMachine->GetDestinationState() > DestinationState)
+		{
+			FGameFeaturePluginUnloadComplete Callback = FGameFeaturePluginUnloadComplete();
+			UnloadGameFeaturePlugin(PluginURL, Callback, bKeepRegistered);
+		}
+	}
 }
 
 void UGameFeaturesSubsystem::UnloadGameFeaturePlugin(const FString& PluginURL, const FGameFeaturePluginUnloadComplete& CompleteDelegate, bool bKeepRegistered)
@@ -546,8 +559,14 @@ void UGameFeaturesSubsystem::UnloadGameFeaturePlugin(const FString& PluginURL, c
 
 void UGameFeaturesSubsystem::UninstallGameFeaturePlugin(const FString& PluginURL)
 {
-	FGameFeaturePluginUninstallComplete Callback = FGameFeaturePluginUninstallComplete();
-	UninstallGameFeaturePlugin(PluginURL, Callback);
+	if (UGameFeaturePluginStateMachine* StateMachine = FindGameFeaturePluginStateMachine(PluginURL))
+	{
+		if (StateMachine->GetDestinationState() > EGameFeaturePluginState::StatusKnown)
+		{
+			FGameFeaturePluginUninstallComplete Callback = FGameFeaturePluginUninstallComplete();
+			UninstallGameFeaturePlugin(PluginURL, Callback);
+		}
+	}
 }
 
 void UGameFeaturesSubsystem::UninstallGameFeaturePlugin(const FString& PluginURL, const FGameFeaturePluginUninstallComplete& CompleteDelegate)
