#UE CIS fix

diff --git a/Engine/Plugins/Experimental/GameFeatures/Source/GameFeatures/Private/GameFeaturesSubsystem.cpp b/Engine/Plugins/Experimental/GameFeatures/Source/GameFeatures/Private/GameFeaturesSubsystem.cpp
--- a/Engine/Plugins/Experimental/GameFeatures/Source/GameFeatures/Private/GameFeaturesSubsystem.cpp
+++ b/Engine/Plugins/Experimental/GameFeatures/Source/GameFeatures/Private/GameFeaturesSubsystem.cpp
@@ -518,7 +518,7 @@ void UGameFeaturesSubsystem::UnloadGameFeaturePlugin(const FString& PluginURL, b
 {
 	if (UGameFeaturePluginStateMachine* StateMachine = FindGameFeaturePluginStateMachine(PluginURL))
 	{
-		EGameFeaturePluginState::Type DestinationState = bKeepRegistered ? EGameFeaturePluginState::Registered : EGameFeaturePluginState::Installed;
+		EGameFeaturePluginState DestinationState = bKeepRegistered ? EGameFeaturePluginState::Registered : EGameFeaturePluginState::Installed;
 		if (StateMachine->GetDestinationState() > DestinationState)
 		{
 			FGameFeaturePluginUnloadComplete Callback = FGameFeaturePluginUnloadComplete();
