#UE Transitioning to an error state when you are waiting for dependencies but they failed to transition to registered

diff --git a/Engine/Plugins/Experimental/GameFeatures/Source/GameFeatures/Private/GameFeaturePluginStateMachine.cpp b/Engine/Plugins/Experimental/GameFeatures/Source/GameFeatures/Private/GameFeaturePluginStateMachine.cpp
--- a/Engine/Plugins/Experimental/GameFeatures/Source/GameFeatures/Private/GameFeaturePluginStateMachine.cpp
+++ b/Engine/Plugins/Experimental/GameFeatures/Source/GameFeatures/Private/GameFeaturePluginStateMachine.cpp
@@ -1015,6 +1015,11 @@ struct FGameFeaturePluginState_WaitingForDependencies : public FGameFeaturePlugi
 				RemainingDependency->OnStateChanged().RemoveAll(this);
 				RemainingDependencies.RemoveAt(DepIdx, 1, false);
 			}
+			else if (RemainingDependency->GetCurrentState() == RemainingDependency->GetDestinationState())
+			{
+				// The dependency is no longer transitioning and is not Registered or later, so it failed to register, thus we cannot proceed
+				StateStatus.SetTransitionError(EGameFeaturePluginState::ErrorWaitingForDependencies, UE::GameFeatures::StateMachineErrorNamespace + TEXT("Failed_Dependency_Register"));
+			}
 		}
 
 		if (RemainingDependencies.Num() == 0)
