Non-Unity fixes

diff --git a/Engine/Plugins/Experimental/GameFeatures/Source/GameFeaturesEditor/Private/GameFeaturePluginMetadataCustomization.cpp b/Engine/Plugins/Experimental/GameFeatures/Source/GameFeaturesEditor/Private/GameFeaturePluginMetadataCustomization.cpp
--- a/Engine/Plugins/Experimental/GameFeatures/Source/GameFeaturesEditor/Private/GameFeaturePluginMetadataCustomization.cpp
+++ b/Engine/Plugins/Experimental/GameFeatures/Source/GameFeaturesEditor/Private/GameFeaturePluginMetadataCustomization.cpp
@@ -1,6 +1,7 @@
 // Copyright Epic Games, Inc. All Rights Reserved.
 
 #include "GameFeaturePluginMetadataCustomization.h"
+#include "Framework/Application/SlateApplication.h"
 #include "Widgets/DeclarativeSyntaxSupport.h"
 #include "Widgets/SBoxPanel.h"
 #include "Widgets/Input/SButton.h"
diff --git a/Engine/Source/Editor/DetailCustomizations/Private/MobilityCustomization.cpp b/Engine/Source/Editor/DetailCustomizations/Private/MobilityCustomization.cpp
--- a/Engine/Source/Editor/DetailCustomizations/Private/MobilityCustomization.cpp
+++ b/Engine/Source/Editor/DetailCustomizations/Private/MobilityCustomization.cpp
@@ -9,8 +9,10 @@
 #include "EditorStyleSet.h"
 #include "DetailLayoutBuilder.h"
 #include "DetailWidgetRow.h"
+#include "Framework/Application/SlateApplication.h"
 #include "Widgets/Input/SSegmentedControl.h"
 
+
 #define LOCTEXT_NAMESPACE "MobilityCustomization" 
 
 FMobilityCustomization::FMobilityCustomization(TSharedPtr<IPropertyHandle> InMobilityHandle, uint8 InRestrictedMobilityBits, bool InForLight)
diff --git a/Engine/Source/Editor/UMGEditor/Private/Customizations/HorizontalAlignmentCustomization.cpp b/Engine/Source/Editor/UMGEditor/Private/Customizations/HorizontalAlignmentCustomization.cpp
--- a/Engine/Source/Editor/UMGEditor/Private/Customizations/HorizontalAlignmentCustomization.cpp
+++ b/Engine/Source/Editor/UMGEditor/Private/Customizations/HorizontalAlignmentCustomization.cpp
@@ -4,6 +4,7 @@
 #include "Widgets/Images/SImage.h"
 #include "Widgets/Input/SCheckBox.h"
 #include "DetailWidgetRow.h"
+#include "Framework/Application/SlateApplication.h"
 #include "Widgets/Input/SSegmentedControl.h"
 
 
diff --git a/Engine/Source/Editor/UMGEditor/Private/Customizations/TextJustifyCustomization.cpp b/Engine/Source/Editor/UMGEditor/Private/Customizations/TextJustifyCustomization.cpp
--- a/Engine/Source/Editor/UMGEditor/Private/Customizations/TextJustifyCustomization.cpp
+++ b/Engine/Source/Editor/UMGEditor/Private/Customizations/TextJustifyCustomization.cpp
@@ -4,6 +4,7 @@
 #include "Widgets/Images/SImage.h"
 #include "Widgets/Input/SCheckBox.h"
 #include "DetailWidgetRow.h"
+#include "Framework/Application/SlateApplication.h"
 #include "Widgets/Input/SSegmentedControl.h"
 
 
diff --git a/Engine/Source/Editor/UMGEditor/Private/Customizations/VerticalAlignmentCustomization.cpp b/Engine/Source/Editor/UMGEditor/Private/Customizations/VerticalAlignmentCustomization.cpp
--- a/Engine/Source/Editor/UMGEditor/Private/Customizations/VerticalAlignmentCustomization.cpp
+++ b/Engine/Source/Editor/UMGEditor/Private/Customizations/VerticalAlignmentCustomization.cpp
@@ -9,6 +9,7 @@
 #include "Widgets/Input/SCheckBox.h"
 #include "PropertyHandle.h"
 #include "DetailWidgetRow.h"
+#include "Framework/Application/SlateApplication.h"
 #include "Widgets/Input/SSegmentedControl.h"
 
 
