Add an explicit apply context to activate/deactivate of game features that lets you specify a specific world context to allow limiting features to a certain PIE instance/world
Change engine and lyra world actions to check the context, the other actions will fall back to the default global context
Change lyra experience manager to let actions apply for each instance, this fixes multi session PIE

diff --git a/Engine/Plugins/Experimental/GameFeatures/Source/GameFeatures/Private/GameFeatureAction.cpp b/Engine/Plugins/Experimental/GameFeatures/Source/GameFeatures/Private/GameFeatureAction.cpp
--- a/Engine/Plugins/Experimental/GameFeatures/Source/GameFeatures/Private/GameFeatureAction.cpp
+++ b/Engine/Plugins/Experimental/GameFeatures/Source/GameFeatures/Private/GameFeatureAction.cpp
@@ -1,4 +1,10 @@
 // Copyright Epic Games, Inc. All Rights Reserved.
 
 #include "GameFeatureAction.h"
+#include "GameFeaturesSubsystem.h"
 
+void UGameFeatureAction::OnGameFeatureActivating(FGameFeatureActivatingContext& Context)
+{
+	// Call older style if not overridden
+	OnGameFeatureActivating();
+}
\ No newline at end of file
diff --git a/Engine/Plugins/Experimental/GameFeatures/Source/GameFeatures/Private/GameFeatureAction_AddComponents.cpp b/Engine/Plugins/Experimental/GameFeatures/Source/GameFeatures/Private/GameFeatureAction_AddComponents.cpp
--- a/Engine/Plugins/Experimental/GameFeatures/Source/GameFeatures/Private/GameFeatureAction_AddComponents.cpp
+++ b/Engine/Plugins/Experimental/GameFeatures/Source/GameFeatures/Private/GameFeatureAction_AddComponents.cpp
@@ -13,25 +13,33 @@
 //////////////////////////////////////////////////////////////////////
 // UGameFeatureAction_AddComponents
 
-void UGameFeatureAction_AddComponents::OnGameFeatureActivating()
+void UGameFeatureAction_AddComponents::OnGameFeatureActivating(FGameFeatureActivatingContext& Context)
 {
-	GameInstanceStartHandle = FWorldDelegates::OnStartGameInstance.AddUObject(this, &UGameFeatureAction_AddComponents::HandleGameInstanceStart);
+	FContextHandles& Handles = ContextHandles.FindOrAdd(Context);
 
-	check(ComponentRequestHandles.Num() == 0);
+	Handles.GameInstanceStartHandle = FWorldDelegates::OnStartGameInstance.AddUObject(this, 
+		&UGameFeatureAction_AddComponents::HandleGameInstanceStart, FGameFeatureStateChangeContext(Context));
+
+	ensure(Handles.ComponentRequestHandles.Num() == 0);
 
 	// Add to any worlds with associated game instances that have already been initialized
 	for (const FWorldContext& WorldContext : GEngine->GetWorldContexts())
 	{
-		AddToWorld(WorldContext);
+		if (Context.ShouldApplyToWorldContext(WorldContext))
+		{
+			AddToWorld(WorldContext, Handles);
+		}
 	}
 }
 
 void UGameFeatureAction_AddComponents::OnGameFeatureDeactivating(FGameFeatureDeactivatingContext& Context)
 {
-	FWorldDelegates::OnStartGameInstance.Remove(GameInstanceStartHandle);
+	FContextHandles& Handles = ContextHandles.FindOrAdd(Context);
+
+	FWorldDelegates::OnStartGameInstance.Remove(Handles.GameInstanceStartHandle);
 
 	// Releasing the handles will also remove the components from any registered actors too
-	ComponentRequestHandles.Empty();
+	Handles.ComponentRequestHandles.Empty();
 }
 
 #if WITH_EDITORONLY_DATA
@@ -81,7 +89,7 @@ EDataValidationResult UGameFeatureAction_AddComponents::IsDataValid(TArray<FText
 }
 #endif
 
-void UGameFeatureAction_AddComponents::AddToWorld(const FWorldContext& WorldContext)
+void UGameFeatureAction_AddComponents::AddToWorld(const FWorldContext& WorldContext, FContextHandles& Handles)
 {
 	UWorld* World = WorldContext.World();
 	UGameInstance* GameInstance = WorldContext.OwningGameInstance;
@@ -106,7 +114,7 @@ void UGameFeatureAction_AddComponents::AddToWorld(const FWorldContext& WorldCont
 						TSubclassOf<UActorComponent> ComponentClass = Entry.ComponentClass.LoadSynchronous();
 						if (ComponentClass)
 						{
-							ComponentRequestHandles.Add(GFCM->AddComponentRequest(Entry.ActorClass, ComponentClass));
+							Handles.ComponentRequestHandles.Add(GFCM->AddComponentRequest(Entry.ActorClass, ComponentClass));
 						}
 						else if (!Entry.ComponentClass.IsNull())
 						{
@@ -119,11 +127,18 @@ void UGameFeatureAction_AddComponents::AddToWorld(const FWorldContext& WorldCont
 	}
 }
 
-void UGameFeatureAction_AddComponents::HandleGameInstanceStart(UGameInstance* GameInstance)
+void UGameFeatureAction_AddComponents::HandleGameInstanceStart(UGameInstance* GameInstance, FGameFeatureStateChangeContext ChangeContext)
 {
 	if (FWorldContext* WorldContext = GameInstance->GetWorldContext())
 	{
-		AddToWorld(*WorldContext);
+		if (ChangeContext.ShouldApplyToWorldContext(*WorldContext))
+		{
+			FContextHandles* Handles = ContextHandles.Find(ChangeContext);
+			if (ensure(Handles))
+			{
+				AddToWorld(*WorldContext, *Handles);
+			}
+		}
 	}
 }
 
diff --git a/Engine/Plugins/Experimental/GameFeatures/Source/GameFeatures/Private/GameFeaturePluginStateMachine.cpp b/Engine/Plugins/Experimental/GameFeatures/Source/GameFeatures/Private/GameFeaturePluginStateMachine.cpp
--- a/Engine/Plugins/Experimental/GameFeatures/Source/GameFeatures/Private/GameFeaturePluginStateMachine.cpp
+++ b/Engine/Plugins/Experimental/GameFeatures/Source/GameFeatures/Private/GameFeaturePluginStateMachine.cpp
@@ -1339,10 +1339,12 @@ struct FGameFeaturePluginState_Activating : public FGameFeaturePluginState
 		check(GEngine);
 		check(StateProperties.GameFeatureData);
 
+		FGameFeatureActivatingContext Context;
+
 		StateProperties.GameFeatureData->InitializeHierarchicalPluginIniFiles(StateProperties.PluginInstalledFilename);
 
 		const FString PluginName = FPaths::GetBaseFilename(StateProperties.PluginInstalledFilename);
-		UGameFeaturesSubsystem::Get().OnGameFeatureActivating(StateProperties.GameFeatureData, PluginName);
+		UGameFeaturesSubsystem::Get().OnGameFeatureActivating(StateProperties.GameFeatureData, PluginName, Context);
 
 		StateStatus.SetTransition(EGameFeaturePluginState::Active);
 	}
diff --git a/Engine/Plugins/Experimental/GameFeatures/Source/GameFeatures/Private/GameFeaturesSubsystem.cpp b/Engine/Plugins/Experimental/GameFeatures/Source/GameFeatures/Private/GameFeaturesSubsystem.cpp
--- a/Engine/Plugins/Experimental/GameFeatures/Source/GameFeatures/Private/GameFeaturesSubsystem.cpp
+++ b/Engine/Plugins/Experimental/GameFeatures/Source/GameFeatures/Private/GameFeaturesSubsystem.cpp
@@ -16,9 +16,44 @@
 #include "Engine/AssetManager.h"
 #include "Engine/StreamableManager.h"
 #include "Engine/AssetManagerSettings.h"
+#include "Engine/Engine.h"
 
 DEFINE_LOG_CATEGORY(LogGameFeatures);
 
+void FGameFeatureStateChangeContext::SetRequiredWorldContextHandle(FName Handle)
+{
+	WorldContextHandle = Handle;
+}
+
+bool FGameFeatureStateChangeContext::ShouldApplyToWorldContext(const FWorldContext& WorldContext) const
+{
+	if (WorldContextHandle.IsNone())
+	{
+		return true;
+	}
+	if (WorldContext.ContextHandle == WorldContextHandle)
+	{
+		return true;
+	}
+	return false;
+}
+
+bool FGameFeatureStateChangeContext::ShouldApplyUsingOtherContext(const FGameFeatureStateChangeContext& OtherContext) const
+{
+	if (OtherContext == *this)
+	{
+		return true;
+	}
+
+	// If other context is less restrictive, apply
+	if (OtherContext.WorldContextHandle.IsNone())
+	{
+		return true;
+	}
+
+	return false;
+}
+
 void UGameFeaturesSubsystem::Initialize(FSubsystemCollectionBase& Collection)
 {
 	UE_LOG(LogGameFeatures, Log, TEXT("Initializing game features subsystem"));
@@ -372,7 +407,7 @@ void UGameFeaturesSubsystem::OnGameFeatureLoading(const UGameFeatureData* GameFe
 	}
 }
 
-void UGameFeaturesSubsystem::OnGameFeatureActivating(const UGameFeatureData* GameFeatureData, const FString& PluginName)
+void UGameFeaturesSubsystem::OnGameFeatureActivating(const UGameFeatureData* GameFeatureData, const FString& PluginName, FGameFeatureActivatingContext& Context)
 {
 	check(GameFeatureData);
 
@@ -380,12 +415,12 @@ void UGameFeaturesSubsystem::OnGameFeatureActivating(const UGameFeatureData* Gam
 	{
 		CastChecked<IGameFeatureStateChangeObserver>(Observer)->OnGameFeatureActivating(GameFeatureData);
 	}
-
+	
 	for (UGameFeatureAction* Action : GameFeatureData->GetActions())
 	{
 		if (Action != nullptr)
 		{
-			Action->OnGameFeatureActivating();
+			Action->OnGameFeatureActivating(Context);
 		}
 	}
 }
@@ -1232,4 +1267,4 @@ EGameFeaturePluginState UGameFeaturesSubsystem::ConvertInitialFeatureStateToTarg
 		break;
 	}
 	return InitialState;
-}
\ No newline at end of file
+}
diff --git a/Engine/Plugins/Experimental/GameFeatures/Source/GameFeatures/Public/GameFeatureAction.h b/Engine/Plugins/Experimental/GameFeatures/Source/GameFeatures/Public/GameFeatureAction.h
--- a/Engine/Plugins/Experimental/GameFeatures/Source/GameFeatures/Public/GameFeatureAction.h
+++ b/Engine/Plugins/Experimental/GameFeatures/Source/GameFeatures/Public/GameFeatureAction.h
@@ -5,6 +5,7 @@
 #include "CoreMinimal.h"
 #include "GameFeatureAction.generated.h"
 
+struct FGameFeatureActivatingContext;
 struct FGameFeatureDeactivatingContext;
 struct FAssetBundleData;
 
@@ -15,14 +16,23 @@ class GAMEFEATURES_API UGameFeatureAction : public UObject
 	GENERATED_BODY()
 
 public:
+
+	/** Called when the object owning the action is registered for possible activation, this is called even if a feature never activates */
 	virtual void OnGameFeatureRegistering() {}
 
+	/** Called to unregister an action, it will not be activated again without being registered again */
 	virtual void OnGameFeatureUnregistering() {}
+	
+	/** Called to indicate that a feature is being loaded for activation in the near future */
+	virtual void OnGameFeatureLoading() {}
 
-	virtual void OnGameFeatureActivating() {}
+	/** Called when the feature is actually applied */
+	virtual void OnGameFeatureActivating(FGameFeatureActivatingContext& Context);
 
-	virtual void OnGameFeatureLoading() {}
+	/** Older-style activation function with no context, called by base class if context version is not overridden */
+	virtual void OnGameFeatureActivating() {}
 
+	/** Called when game feature is deactivated, it may be activated again in the near future */
 	virtual void OnGameFeatureDeactivating(FGameFeatureDeactivatingContext& Context) {}
 
 #if WITH_EDITORONLY_DATA
diff --git a/Engine/Plugins/Experimental/GameFeatures/Source/GameFeatures/Public/GameFeatureAction_AddComponents.h b/Engine/Plugins/Experimental/GameFeatures/Source/GameFeatures/Public/GameFeatureAction_AddComponents.h
--- a/Engine/Plugins/Experimental/GameFeatures/Source/GameFeatures/Public/GameFeatureAction_AddComponents.h
+++ b/Engine/Plugins/Experimental/GameFeatures/Source/GameFeatures/Public/GameFeatureAction_AddComponents.h
@@ -5,6 +5,7 @@
 #include "CoreMinimal.h"
 #include "Engine/World.h"
 #include "GameFeatureAction.h"
+#include "GameFeaturesSubsystem.h"
 
 #include "GameFeatureAction_AddComponents.generated.h"
 
@@ -61,7 +62,7 @@ class UGameFeatureAction_AddComponents final : public UGameFeatureAction
 
 public:
 	//~UGameFeatureAction interface
-	virtual void OnGameFeatureActivating() override;
+	virtual void OnGameFeatureActivating(FGameFeatureActivatingContext& Context) override;
 	virtual void OnGameFeatureDeactivating(FGameFeatureDeactivatingContext& Context) override;
 #if WITH_EDITORONLY_DATA
 	virtual void AddAdditionalAssetBundleData(FAssetBundleData& AssetBundleData) override;
@@ -79,12 +80,15 @@ public:
 	TArray<FGameFeatureComponentEntry> ComponentList;
 
 private:
-	void AddToWorld(const FWorldContext& WorldContext);
+	struct FContextHandles
+	{
+		FDelegateHandle GameInstanceStartHandle;
+		TArray<TSharedPtr<FComponentRequestHandle>> ComponentRequestHandles;
+	};
 
-	void HandleGameInstanceStart(UGameInstance* GameInstance);
+	void AddToWorld(const FWorldContext& WorldContext, FContextHandles& Handles);
 
-private:
-	FDelegateHandle GameInstanceStartHandle;
+	void HandleGameInstanceStart(UGameInstance* GameInstance, FGameFeatureStateChangeContext ChangeContext);
 
-	TArray<TSharedPtr<FComponentRequestHandle>> ComponentRequestHandles;
+	TMap<FGameFeatureStateChangeContext, FContextHandles> ContextHandles;
 };
diff --git a/Engine/Plugins/Experimental/GameFeatures/Source/GameFeatures/Public/GameFeaturesSubsystem.h b/Engine/Plugins/Experimental/GameFeatures/Source/GameFeatures/Public/GameFeaturesSubsystem.h
--- a/Engine/Plugins/Experimental/GameFeatures/Source/GameFeatures/Public/GameFeaturesSubsystem.h
+++ b/Engine/Plugins/Experimental/GameFeatures/Source/GameFeatures/Public/GameFeaturesSubsystem.h
@@ -18,8 +18,60 @@ class UGameFeaturesProjectPolicies;
 enum class EGameFeaturePluginState : uint8;
 class IPlugin;
 class FJsonObject;
+struct FWorldContext;
 
-struct FGameFeatureDeactivatingContext
+/** 
+ * Struct that determines if game feature action state changes should be applied for cases where there are multiple worlds or contexts.
+ * The default value means to apply to all possible objects. This can be safely copied and used for later querying.
+ */
+struct GAMEFEATURES_API FGameFeatureStateChangeContext
+{
+public:
+
+	/** Sets a specific world context handle to limit changes to */
+	void SetRequiredWorldContextHandle(FName Handle);
+
+	/** Sees if the specific world context matches the application rules */
+	bool ShouldApplyToWorldContext(const FWorldContext& WorldContext) const;
+
+	/** True if events bound using this context should apply when using other context */
+	bool ShouldApplyUsingOtherContext(const FGameFeatureStateChangeContext& OtherContext) const;
+
+	/** Check if this has the exact same state change application rules */
+	FORCEINLINE bool operator==(const FGameFeatureStateChangeContext& OtherContext) const
+	{
+		if (OtherContext.WorldContextHandle == WorldContextHandle)
+		{
+			return true;
+		}
+
+		return false;
+	}
+
+	/** Allow this to be used as a map key */
+	FORCEINLINE friend uint32 GetTypeHash(const FGameFeatureStateChangeContext& OtherContext)
+	{
+		return GetTypeHash(OtherContext.WorldContextHandle);
+	}
+
+private:
+	/** Specific world context to limit changes to, if none then it will apply to all */
+	FName WorldContextHandle;
+};
+
+/** Context that provides extra information for activating a game feature */
+struct FGameFeatureActivatingContext : public FGameFeatureStateChangeContext
+{
+public:
+	//@TODO: Add rules specific to activation when required
+
+private:
+
+	friend struct FGameFeaturePluginState_Activating;
+};
+
+/** Context that provides extra information for deactivating a game feature, will use the same change context rules as the activating context */
+struct FGameFeatureDeactivatingContext : public FGameFeatureStateChangeContext
 {
 public:
 	// Call this if your observer has an asynchronous action to complete as part of shutdown, and invoke the returned delegate when you are done (on the game thread!)
@@ -266,7 +318,7 @@ private:
 	void OnGameFeatureUnregistering(const UGameFeatureData* GameFeatureData, const FString& PluginName);
 	friend struct FGameFeaturePluginState_Unregistering;
 
-	void OnGameFeatureActivating(const UGameFeatureData* GameFeatureData, const FString& PluginName);
+	void OnGameFeatureActivating(const UGameFeatureData* GameFeatureData, const FString& PluginName, FGameFeatureActivatingContext& Context);
 	friend struct FGameFeaturePluginState_Activating;
 
 	void OnGameFeatureDeactivating(const UGameFeatureData* GameFeatureData, const FString& PluginName, FGameFeatureDeactivatingContext& Context);
