Game Features: Add tooltips for state switcher widgets and the current state, and fixed up some missing states in the existing display names (along with a static assert to catch future additions)

diff --git a/Engine/Plugins/Experimental/GameFeatures/Source/GameFeatures/Private/GameFeaturePluginStateMachine.h b/Engine/Plugins/Experimental/GameFeatures/Source/GameFeatures/Private/GameFeaturePluginStateMachine.h
--- a/Engine/Plugins/Experimental/GameFeatures/Source/GameFeatures/Private/GameFeaturePluginStateMachine.h
+++ b/Engine/Plugins/Experimental/GameFeatures/Source/GameFeatures/Private/GameFeaturePluginStateMachine.h
@@ -150,9 +150,9 @@ enum class EGameFeaturePluginState : uint8
 	Unregistering,				// Transition state Registered -> Installed. Cleaning up data gathered in Registering.
 	Registering,				// Transition state Installed -> Registered. Discovering assets in the plugin, but not loading them, except a few for discovery reasons.
 	Registered,					// The assets in the plugin are known, but have not yet been loaded, except a few for discovery reasons.
-	Unloading,					// Transition state Loaded -> Registered. In the process of removing code/contnet from memory. 
+	Unloading,					// Transition state Loaded -> Registered. In the process of removing code/content from memory. 
 	Loading,					// Transition state Registered -> Loaded. In the process of loading code/content into memory.
-	Loaded,						// The plugin is loaded into memory, but not registered with game systems and active.
+	Loaded,						// The plugin is loaded into memory and registered with some game systems but not yet active.
 	Deactivating,				// Transition state Active -> Loaded. Currently unregistering with game systems.
 	Activating,					// Transition state Loaded -> Active. Currently registering plugin code/content with game systems.
 	Active,						// Plugin is fully loaded and active. It is affecting the game.
diff --git a/Engine/Plugins/Experimental/GameFeatures/Source/GameFeaturesEditor/Private/SGameFeatureStateWidget.cpp b/Engine/Plugins/Experimental/GameFeatures/Source/GameFeaturesEditor/Private/SGameFeatureStateWidget.cpp
--- a/Engine/Plugins/Experimental/GameFeatures/Source/GameFeaturesEditor/Private/SGameFeatureStateWidget.cpp
+++ b/Engine/Plugins/Experimental/GameFeatures/Source/GameFeaturesEditor/Private/SGameFeatureStateWidget.cpp
@@ -26,15 +26,18 @@ void SGameFeatureStateWidget::Construct(const FArguments& InArgs)
 			SNew(SSegmentedControl<EGameFeaturePluginState>)
 			.Value(CurrentState)
 			.OnValueChanged(InArgs._OnStateChanged)
-//			.ToolTipText(LOCTEXT("StateSwitcherTooltip", "Attempt to change the current state of this game feature"))
 			+SSegmentedControl<EGameFeaturePluginState>::Slot(EGameFeaturePluginState::Installed)
 				.Text(GetDisplayNameOfState(EGameFeaturePluginState::Installed))
+				.ToolTip(LOCTEXT("SwitchToInstalledTooltip", "Attempt to change the current state of this game feature to Installed.\n\nInstalled means that the plugin is in local storage (i.e., it is on the hard drive) but it has not been registered, loaded, or activated yet."))
 			+ SSegmentedControl<EGameFeaturePluginState>::Slot(EGameFeaturePluginState::Registered)
 				.Text(GetDisplayNameOfState(EGameFeaturePluginState::Registered))
+				.ToolTip(LOCTEXT("SwitchToRegisteredTooltip", "Attempt to change the current state of this game feature to Registered.\n\nRegistered means that the assets in the plugin are known, but have not yet been loaded, except a few for discovery reasons (and it is not actively affecting gameplay yet)."))
 			+ SSegmentedControl<EGameFeaturePluginState>::Slot(EGameFeaturePluginState::Loaded)
 				.Text(GetDisplayNameOfState(EGameFeaturePluginState::Loaded))
+				.ToolTip(LOCTEXT("SwitchToLoadedTooltip", "Attempt to change the current state of this game feature to Loaded.\n\nLoaded means that the plugin is in local storage (i.e., it is on the hard drive) but it has not been registered, loaded, or activated yet."))
 			+ SSegmentedControl<EGameFeaturePluginState>::Slot(EGameFeaturePluginState::Active)
 				.Text(GetDisplayNameOfState(EGameFeaturePluginState::Active))
+				.ToolTip(LOCTEXT("SwitchToActiveTooltip", "Attempt to change the current state of this game feature to Active.\n\nActive means that the plugin is fully loaded and active. It is affecting the game."))
 		]
 		+SHorizontalBox::Slot()
 		.Padding(8.0f, 0.0f, 0.0f, 0.0f)
@@ -43,7 +46,7 @@ void SGameFeatureStateWidget::Construct(const FArguments& InArgs)
 			SNew(STextBlock)
 			.Text(this, &SGameFeatureStateWidget::GetStateStatusDisplay)
 			.TextStyle(&FAppStyle::Get().GetWidgetStyle<FTextBlockStyle>("ButtonText"))
-//			.ToolTipText(LOCTEXT("OtherStateToolTip", "The current state of this game feature plugin"))
+			.ToolTipText(this, &SGameFeatureStateWidget::GetStateStatusDisplayTooltip)
 			.ColorAndOpacity(FAppStyle::Get().GetSlateColor(TEXT("Colors.AccentYellow")))
 		]
 	];
@@ -51,15 +54,23 @@ void SGameFeatureStateWidget::Construct(const FArguments& InArgs)
 
 FText SGameFeatureStateWidget::GetDisplayNameOfState(EGameFeaturePluginState State)
 {
+	static_assert((int32)EGameFeaturePluginState::MAX == 26, "");
+
 	switch (State)
 	{
 	case EGameFeaturePluginState::Uninitialized: return LOCTEXT("UninitializedStateDisplayName", "Uninitialized");
+	case EGameFeaturePluginState::Terminal: return LOCTEXT("TerminalDisplayName", "Terminal");
 	case EGameFeaturePluginState::UnknownStatus: return LOCTEXT("UnknownStatusStateDisplayName", "UnknownStatus");
 	case EGameFeaturePluginState::CheckingStatus: return LOCTEXT("CheckingStatusStateDisplayName", "CheckingStatus");
+	case EGameFeaturePluginState::ErrorCheckingStatus: return LOCTEXT("ErrorCheckingStatusDisplayName", "ErrorCheckingStatus");
+	case EGameFeaturePluginState::ErrorUnavailable: return LOCTEXT("ErrorUnavailableDisplayName", "ErrorUnavailable");
 	case EGameFeaturePluginState::StatusKnown: return LOCTEXT("StatusKnownStateDisplayName", "StatusKnown");
+	case EGameFeaturePluginState::ErrorInstalling: return LOCTEXT("ErrorInstallingDisplayName", "ErrorInstalling");
 	case EGameFeaturePluginState::Uninstalling: return LOCTEXT("UninstallingStateDisplayName", "Uninstalling");
 	case EGameFeaturePluginState::Downloading: return LOCTEXT("DownloadingStateDisplayName", "Downloading");
 	case EGameFeaturePluginState::Installed: return LOCTEXT("InstalledStateDisplayName", "Installed");
+	case EGameFeaturePluginState::ErrorMounting: return LOCTEXT("ErrorMountingDisplayName", "ErrorMounting");
+	case EGameFeaturePluginState::ErrorWaitingForDependencies: return LOCTEXT("ErrorWaitingForDependenciesDisplayName", "ErrorWaitingForDependencies");
 	case EGameFeaturePluginState::Unmounting: return LOCTEXT("UnmountingStateDisplayName", "Unmounting");
 	case EGameFeaturePluginState::Mounting: return LOCTEXT("MountingStateDisplayName", "Mounting");
 	case EGameFeaturePluginState::WaitingForDependencies: return LOCTEXT("WaitingForDependenciesStateDisplayName", "WaitingForDependencies");
@@ -73,10 +84,74 @@ FText SGameFeatureStateWidget::GetDisplayNameOfState(EGameFeaturePluginState Sta
 	case EGameFeaturePluginState::Activating: return LOCTEXT("ActivatingStateDisplayName", "Activating");
 	case EGameFeaturePluginState::Active: return LOCTEXT("ActiveStateDisplayName", "Active");
 	case EGameFeaturePluginState::ErrorRegistering: return LOCTEXT("ErrorRegisteringDisplayName", "ErrorRegistering");
-	default:
-		check(0);
-		return FText::GetEmpty();
 	}
+
+	const FString FallbackString = UE::GameFeatures::ToString(State);
+	ensureMsgf(false, TEXT("Unknown EGameFeaturePluginState entry %d %s"), (int)State, *FallbackString);
+	return FText::AsCultureInvariant(FallbackString);
+}
+
+FText SGameFeatureStateWidget::GetTooltipOfState(EGameFeaturePluginState State)
+{
+	static_assert((int32)EGameFeaturePluginState::MAX == 26, "");
+
+	switch (State)
+	{
+	case EGameFeaturePluginState::Uninitialized:
+		return LOCTEXT("StateTooltip_Uninitialized", "Unset. Not yet been set up.");
+	case EGameFeaturePluginState::Terminal:
+		return LOCTEXT("StateTooltip_Terminal", "Final State before removal of the state machine");
+	case EGameFeaturePluginState::UnknownStatus:
+		return LOCTEXT("StateTooltip_UnknownStatus", "Initialized, but the only thing known is the URL to query status.");
+	case EGameFeaturePluginState::CheckingStatus:
+		return LOCTEXT("StateTooltip_CheckingStatus", "Transition state UnknownStatus -> StatusKnown. The status is in the process of being queried.");
+	case EGameFeaturePluginState::ErrorCheckingStatus:
+		return LOCTEXT("StateTooltip_ErrorCheckingStatus", "Error state for UnknownStatus -> StatusKnown transition.");
+	case EGameFeaturePluginState::ErrorUnavailable:
+		return LOCTEXT("StateTooltip_ErrorUnavailable", "Error state for UnknownStatus -> StatusKnown transition.");
+	case EGameFeaturePluginState::StatusKnown:
+		return LOCTEXT("StateTooltip_StatusKnown", "The plugin's information is known, but no action has taken place yet.");
+	case EGameFeaturePluginState::ErrorInstalling:
+		return LOCTEXT("StateTooltip_ErrorInstalling", "Error state for Installed -> StatusKnown and StatusKnown -> Installed transitions.");
+	case EGameFeaturePluginState::Uninstalling:
+		return LOCTEXT("StateTooltip_Uninstalling", "Transition state Installed -> StatusKnown. In the process of removing from local storage.");
+	case EGameFeaturePluginState::Downloading:
+		return LOCTEXT("StateTooltip_Downloading", "Transition state StatusKnown -> Installed. In the process of adding to local storage.");
+	case EGameFeaturePluginState::Installed:
+		return LOCTEXT("StateTooltip_Installed", "The plugin is in local storage (i.e. it is on the hard drive)");
+	case EGameFeaturePluginState::ErrorMounting:
+		return LOCTEXT("StateTooltip_ErrorMounting", "Error state for Installed -> Registered and Registered -> Installed transitions.");
+	case EGameFeaturePluginState::ErrorWaitingForDependencies:
+		return LOCTEXT("StateTooltip_ErrorWaitingForDependencies", "Error state for Installed -> Registered and Registered -> Installed transitions.");
+	case EGameFeaturePluginState::ErrorRegistering:
+		return LOCTEXT("StateTooltip_ErrorRegistering", "Error state for Installed -> Registered and Registered -> Installed transitions.");
+	case EGameFeaturePluginState::WaitingForDependencies:
+		return LOCTEXT("StateTooltip_WaitingForDependencies", "Transition state Installed -> Registered. In the process of loading code/content for all dependencies into memory.");
+	case EGameFeaturePluginState::Unmounting:
+		return LOCTEXT("StateTooltip_Unmounting", "Transition state Registered -> Installed. The content file(s) (i.e. pak file) for the plugin is unmounting.");
+	case EGameFeaturePluginState::Mounting:
+		return LOCTEXT("StateTooltip_Mounting", "Transition state Installed -> Registered. The content files(s) (i.e. pak file) for the plugin is getting mounted.");
+	case EGameFeaturePluginState::Unregistering:
+		return LOCTEXT("StateTooltip_Unregistering", "Transition state Registered -> Installed. Cleaning up data gathered in Registering.");
+	case EGameFeaturePluginState::Registering:
+		return LOCTEXT("StateTooltip_Registering", "Transition state Installed -> Registered. Discovering assets in the plugin, but not loading them, except a few for discovery reasons.");
+	case EGameFeaturePluginState::Registered:
+		return LOCTEXT("StateTooltip_Registered", "The assets in the plugin are known, but have not yet been loaded, except a few for discovery reasons.");
+	case EGameFeaturePluginState::Unloading:
+		return LOCTEXT("StateTooltip_Unloading", "Transition state Loaded -> Registered. In the process of removing code/content from memory.");
+	case EGameFeaturePluginState::Loading:
+		return LOCTEXT("StateTooltip_Loading", "Transition state Registered -> Loaded. In the process of loading code/content into memory.");
+	case EGameFeaturePluginState::Loaded:
+		return LOCTEXT("StateTooltip_Loaded", "The plugin is loaded into memory and registered with some game systems but not yet active.");
+	case EGameFeaturePluginState::Deactivating:
+		return LOCTEXT("StateTooltip_Deactivating", "Transition state Active -> Loaded. Currently unregistering with game systems.");
+	case EGameFeaturePluginState::Activating:
+		return LOCTEXT("StateTooltip_Activating", "Transition state Loaded -> Active. Currently registering plugin code/content with game systems.");
+	case EGameFeaturePluginState::Active:
+		return LOCTEXT("StateTooltip_Active", "Plugin is fully loaded and active. It is affecting the game.");
+	}
+
+	return GetDisplayNameOfState(State);
 }
 
 FText SGameFeatureStateWidget::GetStateStatusDisplay() const
@@ -95,6 +170,15 @@ FText SGameFeatureStateWidget::GetStateStatusDisplay() const
 	}
 }
 
+FText SGameFeatureStateWidget::GetStateStatusDisplayTooltip() const
+{
+	const EGameFeaturePluginState State = CurrentState.Get();
+
+	return FText::Format(
+		LOCTEXT("OtherStateToolTip", "The current state of this game feature plugin\n\n{0}"),
+		GetTooltipOfState(State));
+}
+
 //////////////////////////////////////////////////////////////////////////
 
 #undef LOCTEXT_NAMESPACE
diff --git a/Engine/Plugins/Experimental/GameFeatures/Source/GameFeaturesEditor/Private/SGameFeatureStateWidget.h b/Engine/Plugins/Experimental/GameFeatures/Source/GameFeaturesEditor/Private/SGameFeatureStateWidget.h
--- a/Engine/Plugins/Experimental/GameFeatures/Source/GameFeaturesEditor/Private/SGameFeatureStateWidget.h
+++ b/Engine/Plugins/Experimental/GameFeatures/Source/GameFeaturesEditor/Private/SGameFeatureStateWidget.h
@@ -32,9 +32,11 @@ public:
 	void Construct(const FArguments& InArgs);
 
 	static FText GetDisplayNameOfState(EGameFeaturePluginState State);
+	static FText GetTooltipOfState(EGameFeaturePluginState StateID);
 
 private:
 	FText GetStateStatusDisplay() const;
+	FText GetStateStatusDisplayTooltip() const;
 
 private:
 	TAttribute<EGameFeaturePluginState> CurrentState;
