Game Features: Allow changing from Installed to Registered in the game feature switcher UI (all combinations of src->dst for 'target' states should now work)

diff --git a/Engine/Plugins/Experimental/GameFeatures/Source/GameFeaturesEditor/Private/GameFeatureDataDetailsCustomization.cpp b/Engine/Plugins/Experimental/GameFeatures/Source/GameFeaturesEditor/Private/GameFeatureDataDetailsCustomization.cpp
--- a/Engine/Plugins/Experimental/GameFeatures/Source/GameFeaturesEditor/Private/GameFeatureDataDetailsCustomization.cpp
+++ b/Engine/Plugins/Experimental/GameFeatures/Source/GameFeaturesEditor/Private/GameFeatureDataDetailsCustomization.cpp
@@ -122,7 +122,6 @@ void FGameFeatureDataDetailsCustomization::CustomizeDetails(IDetailLayoutBuilder
 				.AutoHeight()
 				[
 					SNew(SGameFeatureStateWidget)
-					.ToolTipText(LOCTEXT("StateSwitcherTooltip", "Attempt to change the current state of this game feature"))
 					.CurrentState(this, &FGameFeatureDataDetailsCustomization::GetCurrentState)
 					.OnStateChanged(this, &FGameFeatureDataDetailsCustomization::ChangeDesiredState)
 				]
@@ -186,45 +185,28 @@ void FGameFeatureDataDetailsCustomization::CustomizeDetails(IDetailLayoutBuilder
 
 void FGameFeatureDataDetailsCustomization::ChangeDesiredState(EGameFeaturePluginState DesiredState)
 {
+	EGameFeatureTargetState TargetState = EGameFeatureTargetState::Installed;
+	switch (DesiredState)
+	{
+	case EGameFeaturePluginState::Installed:
+		TargetState = EGameFeatureTargetState::Installed;
+		break;
+	case EGameFeaturePluginState::Registered:
+		TargetState = EGameFeatureTargetState::Registered;
+		break;
+	case EGameFeaturePluginState::Loaded:
+		TargetState = EGameFeatureTargetState::Loaded;
+		break;
+	case EGameFeaturePluginState::Active:
+		TargetState = EGameFeatureTargetState::Active;
+		break;
+	}
+
 	ErrorTextWidget->SetError(FText::GetEmpty());
 	const TWeakPtr<FGameFeatureDataDetailsCustomization> WeakThisPtr = StaticCastSharedRef<FGameFeatureDataDetailsCustomization>(AsShared());
 
-
 	UGameFeaturesSubsystem& Subsystem = UGameFeaturesSubsystem::Get();
-
-	const EGameFeaturePluginState CurrentState = GetCurrentState();
-	if (DesiredState == EGameFeaturePluginState::Active)
-	{
-		Subsystem.LoadAndActivateGameFeaturePlugin(PluginURL, FGameFeaturePluginLoadComplete::CreateStatic(&FGameFeatureDataDetailsCustomization::OnOperationCompletedOrFailed, WeakThisPtr));
-	}
-	else if (DesiredState == EGameFeaturePluginState::Loaded)
-	{
-		if (CurrentState < EGameFeaturePluginState::Loaded)
-		{
-			Subsystem.LoadGameFeaturePlugin(PluginURL, FGameFeaturePluginLoadComplete::CreateStatic(&FGameFeatureDataDetailsCustomization::OnOperationCompletedOrFailed, WeakThisPtr));
-		}
-		else
-		{
-			Subsystem.DeactivateGameFeaturePlugin(PluginURL, FGameFeaturePluginDeactivateComplete::CreateStatic(&FGameFeatureDataDetailsCustomization::OnOperationCompletedOrFailed, WeakThisPtr));
-		}
-	}
-	else if (DesiredState == EGameFeaturePluginState::Registered)
-	{
-		if (CurrentState >= EGameFeaturePluginState::Loaded)
-		{
-			Subsystem.UnloadGameFeaturePlugin(PluginURL, FGameFeaturePluginDeactivateComplete::CreateStatic(&FGameFeatureDataDetailsCustomization::OnOperationCompletedOrFailed, WeakThisPtr), /*bKeepRegistered=*/ true);
-		}
-		else
-		{
-			//@TODO: No public transition from Installed..Registered is exposed yet
-		}
-	}
-	else if (DesiredState == EGameFeaturePluginState::Installed)
-	{
-		//@TODO: No public transition from something greater than Installed to Installed is exposed yet
-		//@TODO: Do we need to support unregistering?  If not, should remove this button
-		Subsystem.UnloadGameFeaturePlugin(PluginURL, FGameFeaturePluginDeactivateComplete::CreateStatic(&FGameFeatureDataDetailsCustomization::OnOperationCompletedOrFailed, WeakThisPtr), /*bKeepRegistered=*/ false);
-	}
+	Subsystem.ChangeGameFeatureTargetState(PluginURL, TargetState, FGameFeaturePluginDeactivateComplete::CreateStatic(&FGameFeatureDataDetailsCustomization::OnOperationCompletedOrFailed, WeakThisPtr));
 }
 
 
