#UE Reduce log spam. Made game feature state machine log messages verbose, and made the "load complete" message include the ending state.

diff --git a/Engine/Plugins/Experimental/GameFeatures/Source/GameFeatures/Private/GameFeaturePluginStateMachine.cpp b/Engine/Plugins/Experimental/GameFeatures/Source/GameFeatures/Private/GameFeaturePluginStateMachine.cpp
--- a/Engine/Plugins/Experimental/GameFeatures/Source/GameFeatures/Private/GameFeaturePluginStateMachine.cpp
+++ b/Engine/Plugins/Experimental/GameFeatures/Source/GameFeatures/Private/GameFeaturePluginStateMachine.cpp
@@ -1558,7 +1558,7 @@ void UGameFeaturePluginStateMachine::UpdateStateMachine()
 
 		if (StateStatus.TransitionToState != EGameFeaturePluginState::Uninitialized)
 		{
-			UE_LOG(LogGameFeatures, Display, TEXT("Game feature '%s' transitioning state (%s -> %s)"), *GetGameFeatureName(), *UE::GameFeatures::ToString(CurrentState), *UE::GameFeatures::ToString(StateStatus.TransitionToState));
+			UE_LOG(LogGameFeatures, Verbose, TEXT("Game feature '%s' transitioning state (%s -> %s)"), *GetGameFeatureName(), *UE::GameFeatures::ToString(CurrentState), *UE::GameFeatures::ToString(StateStatus.TransitionToState));
 			AllStates[(int32)CurrentState]->EndState();
 			CurrentStateInfo = FGameFeaturePluginStateInfo(StateStatus.TransitionToState);
 			CurrentState = StateStatus.TransitionToState;
diff --git a/Engine/Plugins/Experimental/GameFeatures/Source/GameFeatures/Private/GameFeaturesSubsystem.cpp b/Engine/Plugins/Experimental/GameFeatures/Source/GameFeatures/Private/GameFeaturesSubsystem.cpp
--- a/Engine/Plugins/Experimental/GameFeatures/Source/GameFeatures/Private/GameFeaturesSubsystem.cpp
+++ b/Engine/Plugins/Experimental/GameFeatures/Source/GameFeatures/Private/GameFeaturesSubsystem.cpp
@@ -1025,7 +1025,7 @@ void UGameFeaturesSubsystem::LoadGameFeaturePluginComplete(UGameFeaturePluginSta
 	check(Machine);
 	if (Result.HasValue())
 	{
-		UE_LOG(LogGameFeatures, Display, TEXT("Game feature '%s' loaded successfully."), *Machine->GetGameFeatureName());
+		UE_LOG(LogGameFeatures, Display, TEXT("Game feature '%s' loaded successfully. Ending state: %s"), *Machine->GetGameFeatureName(), *UE::GameFeatures::ToString(Machine->GetCurrentState()));
 	}
 	else
 	{
