Game Features: Changed the user facing display name for GameFeaturesManagerClassName to "Game Feature Project Policy Class"

diff --git a/Engine/Plugins/Experimental/GameFeatures/Source/GameFeatures/Public/GameFeaturesSubsystemSettings.h b/Engine/Plugins/Experimental/GameFeatures/Source/GameFeatures/Public/GameFeaturesSubsystemSettings.h
--- a/Engine/Plugins/Experimental/GameFeatures/Source/GameFeatures/Public/GameFeaturesSubsystemSettings.h
+++ b/Engine/Plugins/Experimental/GameFeatures/Source/GameFeatures/Public/GameFeaturesSubsystemSettings.h
@@ -22,8 +22,8 @@ public:
 	/** State/Bundle to always load on dedicated server */
 	static const FName LoadStateServer;
 
-	/** Name of a singleton class to spawn as the AssetManager, configurable per game. If empty, it will spawn the default one (UDefaultGameFeaturesProjectPolicies) */
-	UPROPERTY(config, EditAnywhere, Category=DefaultClasses, meta=(MetaClass="GameFeaturesProjectPolicies", DisplayName="Game Feature Manager Class", ConfigRestartRequired=true))
+	/** Name of a singleton class to spawn as the game feature project policy. If empty, it will spawn the default one (UDefaultGameFeaturesProjectPolicies) */
+	UPROPERTY(config, EditAnywhere, Category=DefaultClasses, meta=(MetaClass="GameFeaturesProjectPolicies", DisplayName="Game Feature Project Policy Class", ConfigRestartRequired=true))
 	FSoftClassPath GameFeaturesManagerClassName;
 
 	/** List of plugins that are forcibly disabled (e.g., via a hotfix) */
