content worker compiler fix.

diff --git a/Engine/Plugins/Experimental/GameFeatures/Source/GameFeaturesEditor/Private/GameFeaturesEditorModule.cpp b/Engine/Plugins/Experimental/GameFeatures/Source/GameFeaturesEditor/Private/GameFeaturesEditorModule.cpp
--- a/Engine/Plugins/Experimental/GameFeatures/Source/GameFeaturesEditor/Private/GameFeaturesEditorModule.cpp
+++ b/Engine/Plugins/Experimental/GameFeatures/Source/GameFeaturesEditor/Private/GameFeaturesEditorModule.cpp
@@ -246,7 +246,7 @@ class FGameFeaturesEditorModule : public FDefaultModuleImpl
 		if (IModularFeatures::Get().IsModularFeatureAvailable(EditorFeatures::PluginsEditor))
 		{
 			IPluginsEditorFeature& PluginEditor = IModularFeatures::Get().GetModularFeature<IPluginsEditorFeature>(EditorFeatures::PluginsEditor);
-			for (const TSharedPtr<FPluginTemplateDescription> TemplateDescription : PluginTemplates)
+			for (const TSharedPtr<FGameFeaturePluginTemplateDescription, ESPMode::ThreadSafe>& TemplateDescription : PluginTemplates)
 			{
 				PluginEditor.RegisterPluginTemplate(TemplateDescription.ToSharedRef());
 			}
@@ -259,7 +259,7 @@ class FGameFeaturesEditorModule : public FDefaultModuleImpl
 		if (IModularFeatures::Get().IsModularFeatureAvailable(EditorFeatures::PluginsEditor))
 		{
 			IPluginsEditorFeature& PluginEditor = IModularFeatures::Get().GetModularFeature<IPluginsEditorFeature>(EditorFeatures::PluginsEditor);
-			for (const TSharedPtr<FPluginTemplateDescription> TemplateDescription : PluginTemplates)
+			for (const TSharedPtr<FGameFeaturePluginTemplateDescription, ESPMode::ThreadSafe>& TemplateDescription : PluginTemplates)
 			{
 				PluginEditor.UnregisterPluginTemplate(TemplateDescription.ToSharedRef());
 			}
