Redo fix to allow game features to load both absolute and plugin-relative paths for asset manager types
Reimplementation of CL 15419795 which got lost in a merge, was pointed out by a licensee

diff --git a/Engine/Plugins/Experimental/GameFeatures/Source/GameFeatures/Private/GameFeaturesSubsystem.cpp b/Engine/Plugins/Experimental/GameFeatures/Source/GameFeatures/Private/GameFeaturesSubsystem.cpp
--- a/Engine/Plugins/Experimental/GameFeatures/Source/GameFeatures/Private/GameFeaturesSubsystem.cpp
+++ b/Engine/Plugins/Experimental/GameFeatures/Source/GameFeatures/Private/GameFeaturesSubsystem.cpp
@@ -236,14 +236,16 @@ void UGameFeaturesSubsystem::AddGameFeatureToAssetManager(const UGameFeatureData
 {
 	check(GameFeatureToAdd);
 	UAssetManager& LocalAssetManager = UAssetManager::Get();
-
 	LocalAssetManager.PushBulkScanning();
 
+	FString PluginRootPath = TEXT("/") + PluginName + TEXT("/");
+
 	for (FPrimaryAssetTypeInfo TypeInfo : GameFeatureToAdd->GetPrimaryAssetTypesToScan())
 	{
 		for (FDirectoryPath& Path : TypeInfo.Directories)
 		{
-			Path.Path = TEXT("/") + PluginName + TEXT("/") + Path.Path;
+			// Convert plugin-relative paths to full package paths
+			FixPluginPackagePath(Path.Path, PluginRootPath, false);
 		}
 
 		// This function also fills out runtime data on the copy
