Undo //UE5/Release-5.0/Engine/Plugins/Experimental/GameFeatures/Source/GameFeatures/Private/... changelist 18679680
Backout for Lyra Cook failure. Now results in a warning instead..

diff --git a/Engine/Plugins/Experimental/GameFeatures/Source/GameFeatures/Private/GameFeaturesSubsystem.cpp b/Engine/Plugins/Experimental/GameFeatures/Source/GameFeatures/Private/GameFeaturesSubsystem.cpp
--- a/Engine/Plugins/Experimental/GameFeatures/Source/GameFeatures/Private/GameFeaturesSubsystem.cpp
+++ b/Engine/Plugins/Experimental/GameFeatures/Source/GameFeatures/Private/GameFeaturesSubsystem.cpp
@@ -236,16 +236,14 @@ void UGameFeaturesSubsystem::AddGameFeatureToAssetManager(const UGameFeatureData
 {
 	check(GameFeatureToAdd);
 	UAssetManager& LocalAssetManager = UAssetManager::Get();
-	LocalAssetManager.PushBulkScanning();
 
-	FString PluginRootPath = TEXT("/") + PluginName + TEXT("/");
+	LocalAssetManager.PushBulkScanning();
 
 	for (FPrimaryAssetTypeInfo TypeInfo : GameFeatureToAdd->GetPrimaryAssetTypesToScan())
 	{
 		for (FDirectoryPath& Path : TypeInfo.Directories)
 		{
-			// Convert plugin-relative paths to full package paths
-			FixPluginPackagePath(Path.Path, PluginRootPath, false);
+			Path.Path = TEXT("/") + PluginName + TEXT("/") + Path.Path;
 		}
 
 		// This function also fills out runtime data on the copy
