Use plugin name or override name not the GameFeature name.  Was changed in CL 18653855.  IAssetTools::CreateAsset() [UAssetToolsImpl::CreateAsset()] expects the first parameter to be a name not feature type description.

diff --git a/Engine/Plugins/Experimental/GameFeatures/Source/GameFeaturesEditor/Private/GameFeaturesEditorModule.cpp b/Engine/Plugins/Experimental/GameFeatures/Source/GameFeaturesEditor/Private/GameFeaturesEditorModule.cpp
--- a/Engine/Plugins/Experimental/GameFeatures/Source/GameFeaturesEditor/Private/GameFeaturesEditorModule.cpp
+++ b/Engine/Plugins/Experimental/GameFeatures/Source/GameFeaturesEditor/Private/GameFeaturesEditorModule.cpp
@@ -36,7 +36,7 @@ struct FGameFeaturePluginTemplateDescription : public FPluginTemplateDescription
 	{
 		SortPriority = 10;
 		bCanBePlacedInEngine = false;
-		GameFeatureDataName = !GameFeatureDataNameOverride.IsEmpty() ? GameFeatureDataNameOverride : InName.ToString();
+		GameFeatureDataName = !GameFeatureDataNameOverride.IsEmpty() ? GameFeatureDataNameOverride : FString();
 		GameFeatureDataClass = GameFeatureDataClassOverride != nullptr ? GameFeatureDataClassOverride : TSubclassOf<UGameFeatureData>(UGameFeatureData::StaticClass());
 	}
 
@@ -97,7 +97,8 @@ struct FGameFeaturePluginTemplateDescription : public FPluginTemplateDescription
 		{
 			// Create the game feature data asset
 			FAssetToolsModule& AssetToolsModule = FModuleManager::Get().LoadModuleChecked<FAssetToolsModule>("AssetTools");
-			GameFeatureDataAsset = AssetToolsModule.Get().CreateAsset(GameFeatureDataName, NewPlugin->GetMountedAssetPath(), GameFeatureDataClass, /*Factory=*/ nullptr);
+			FString const& AssetName = !GameFeatureDataName.IsEmpty() ? GameFeatureDataName : NewPlugin->GetName();
+			GameFeatureDataAsset = AssetToolsModule.Get().CreateAsset(AssetName, NewPlugin->GetMountedAssetPath(), GameFeatureDataClass, /*Factory=*/ nullptr);
 		}
 		else
 		{
