Prevent a crash on shutdown if the UObject system has already been torn down by the time we shut down the game features editor

diff --git a/Engine/Plugins/Experimental/GameFeatures/Source/GameFeaturesEditor/Private/GameFeaturesEditorModule.cpp b/Engine/Plugins/Experimental/GameFeatures/Source/GameFeaturesEditor/Private/GameFeaturesEditorModule.cpp
--- a/Engine/Plugins/Experimental/GameFeatures/Source/GameFeaturesEditor/Private/GameFeaturesEditorModule.cpp
+++ b/Engine/Plugins/Experimental/GameFeatures/Source/GameFeaturesEditor/Private/GameFeaturesEditorModule.cpp
@@ -184,6 +184,7 @@ class FGameFeaturesEditorModule : public FDefaultModuleImpl
 		}
 
 		// Remove the plugin wizard override
+		if (UObjectInitialized())
 		{
 			if (UClass* GameFeatureClass = UGameFeaturesEditorSettings::StaticClass())
 			{
