Made game feature plugins created via the plugin wizard set EnabledByDefault=false, since the game feature subsystem will manage when they are enabled or not

diff --git a/Engine/Plugins/Experimental/GameFeatures/Source/GameFeaturesEditor/Private/GameFeaturesEditorModule.cpp b/Engine/Plugins/Experimental/GameFeatures/Source/GameFeaturesEditor/Private/GameFeaturesEditorModule.cpp
--- a/Engine/Plugins/Experimental/GameFeatures/Source/GameFeaturesEditor/Private/GameFeaturesEditorModule.cpp
+++ b/Engine/Plugins/Experimental/GameFeatures/Source/GameFeaturesEditor/Private/GameFeaturesEditorModule.cpp
@@ -72,6 +72,9 @@ struct FGameFeaturePluginTemplateDescription : public FPluginTemplateDescription
 		Descriptor.AdditionalFieldsToWrite.FindOrAdd(TEXT("BuiltInInitialFeatureState")) = MakeShared<FJsonValueString>(TEXT("Active"));
 		Descriptor.Category = TEXT("Game Features");
 
+		// Game features should be disabled by default, as they are managed by the game feature subsystem
+		Descriptor.EnabledByDefault = EPluginEnabledByDefault::Disabled;
+
 		if (Descriptor.Modules.Num() > 0)
 		{
 			Descriptor.Modules[0].Name = FName(*(Descriptor.Modules[0].Name.ToString() + TEXT("Runtime")));
