Added support for game feature plugins being found or created in extension directories

diff --git a/Engine/Plugins/Experimental/GameFeatures/Source/GameFeatures/Private/GameFeaturesSubsystem.cpp b/Engine/Plugins/Experimental/GameFeatures/Source/GameFeatures/Private/GameFeaturesSubsystem.cpp
--- a/Engine/Plugins/Experimental/GameFeatures/Source/GameFeatures/Private/GameFeaturesSubsystem.cpp
+++ b/Engine/Plugins/Experimental/GameFeatures/Source/GameFeatures/Private/GameFeaturesSubsystem.cpp
@@ -769,9 +769,8 @@ void UGameFeaturesSubsystem::LoadBuiltInGameFeaturePlugin(const TSharedRef<IPlug
 
 	const FString& PluginDescriptorFilename = Plugin->GetDescriptorFileName();
 
-	// Make sure you are in the game feature plugins folder. All GameFeaturePlugins are in this folder.
-	//@TODO: GameFeaturePluginEnginePush: Comments elsewhere allow plugins outside of the folder as long as they explicitly opt in, either those are wrong or this check is wrong
-	if (!PluginDescriptorFilename.IsEmpty() && FPaths::ConvertRelativePathToFull(PluginDescriptorFilename).StartsWith(GetDefault<UGameFeaturesSubsystemSettings>()->BuiltInGameFeaturePluginsFolder) && FPaths::FileExists(PluginDescriptorFilename))
+	// Make sure you are in a game feature plugins folder. All GameFeaturePlugins are rooted in a GameFeatures folder.
+	if (!PluginDescriptorFilename.IsEmpty() && GetDefault<UGameFeaturesSubsystemSettings>()->IsValidGameFeaturePlugin(FPaths::ConvertRelativePathToFull(PluginDescriptorFilename)) && FPaths::FileExists(PluginDescriptorFilename))
 	{
 		const FString PluginURL = GetPluginURL_FileProtocol(PluginDescriptorFilename);
 		if (GameSpecificPolicies->IsPluginAllowed(PluginURL))
@@ -961,7 +960,7 @@ bool UGameFeaturesSubsystem::GetGameFeaturePluginDetails(const FString& PluginDe
 								const FString& PluginDependencyDescriptorFilename = DependencyPlugin->GetDescriptorFileName();
 
 								if (!PluginDependencyDescriptorFilename.IsEmpty() &&
-									FPaths::ConvertRelativePathToFull(PluginDependencyDescriptorFilename).StartsWith(GetDefault<UGameFeaturesSubsystemSettings>()->BuiltInGameFeaturePluginsFolder) &&
+									GetDefault<UGameFeaturesSubsystemSettings>()->IsValidGameFeaturePlugin(FPaths::ConvertRelativePathToFull(PluginDependencyDescriptorFilename)) &&
 									FPaths::FileExists(PluginDependencyDescriptorFilename))
 								{
 									OutPluginDetails.PluginDependencies.Add(GetPluginURL_FileProtocol(DependencyPlugin->GetDescriptorFileName()));
diff --git a/Engine/Plugins/Experimental/GameFeatures/Source/GameFeatures/Private/GameFeaturesSubsystemSettings.cpp b/Engine/Plugins/Experimental/GameFeatures/Source/GameFeatures/Private/GameFeaturesSubsystemSettings.cpp
--- a/Engine/Plugins/Experimental/GameFeatures/Source/GameFeatures/Private/GameFeaturesSubsystemSettings.cpp
+++ b/Engine/Plugins/Experimental/GameFeatures/Source/GameFeatures/Private/GameFeaturesSubsystemSettings.cpp
@@ -7,6 +7,35 @@ const FName UGameFeaturesSubsystemSettings::LoadStateClient(TEXT("Client"));
 const FName UGameFeaturesSubsystemSettings::LoadStateServer(TEXT("Server"));
 
 UGameFeaturesSubsystemSettings::UGameFeaturesSubsystemSettings()
-	: BuiltInGameFeaturePluginsFolder(FPaths::ConvertRelativePathToFull(FPaths::ProjectPluginsDir() + TEXT("GameFeatures/")))
 {
+	PRAGMA_DISABLE_DEPRECATION_WARNINGS
+	BuiltInGameFeaturePluginsFolder = FPaths::ConvertRelativePathToFull(FPaths::ProjectPluginsDir() + TEXT("GameFeatures/"));
+	PRAGMA_ENABLE_DEPRECATION_WARNINGS
+}
+
+bool UGameFeaturesSubsystemSettings::IsValidGameFeaturePlugin(const FString& PluginDescriptorFilename) const
+{
+	// Build the cache of game feature plugin folders the first time this is called
+	if (BuiltInGameFeaturePluginsFolders.IsEmpty())
+	{
+		BuiltInGameFeaturePluginsFolders.Append(
+			FPaths::GetExtensionDirs(FPaths::ProjectDir(), FPaths::Combine(TEXT("Plugins"), TEXT("GameFeatures")))
+		);
+
+		for (FString& BuiltInFolder : BuiltInGameFeaturePluginsFolders)
+		{
+			BuiltInFolder = FPaths::ConvertRelativePathToFull(BuiltInFolder);
+		}
+	}
+
+	// Check to see if the filename is rooted in a game feature plugin folder
+	for (const FString& BuiltInFolder : BuiltInGameFeaturePluginsFolders)
+	{
+		if (PluginDescriptorFilename.StartsWith(BuiltInFolder))
+		{
+			return true;
+		}
+	}
+
+	return false;
 }
diff --git a/Engine/Plugins/Experimental/GameFeatures/Source/GameFeatures/Public/GameFeaturesSubsystemSettings.h b/Engine/Plugins/Experimental/GameFeatures/Source/GameFeatures/Public/GameFeaturesSubsystemSettings.h
--- a/Engine/Plugins/Experimental/GameFeatures/Source/GameFeatures/Public/GameFeaturesSubsystemSettings.h
+++ b/Engine/Plugins/Experimental/GameFeatures/Source/GameFeatures/Public/GameFeaturesSubsystemSettings.h
@@ -34,9 +34,15 @@ public:
 	UPROPERTY(config, EditAnywhere, Category=GameFeatures)
 	TArray<FString> AdditionalPluginMetadataKeys;
 
-	/** The folder in which all discovered plugins are automatically considered game feature plugins. Plugins outside this folder may also be game features, but need to have bGameFeature: true in the plugin file */
-	//@TODO: GameFeaturePluginEnginePush: Make this configurable
-	//@TODO: GameFeaturePluginEnginePush: This comment doesn't jive with some of the code in the subsystem which is only paying attention to plugins in this folder
-	UPROPERTY(transient)
+	UE_DEPRECATED(5.0, "Use IsValidGameFeaturePlugin() instead")
 	FString BuiltInGameFeaturePluginsFolder;
+
+private:
+	// Cached list of all GameFeature plugin directories (including extension versions)
+	mutable TArray<FString> BuiltInGameFeaturePluginsFolders;
+
+public:
+	// Returns true if the specified (fully qualified) path is a game feature plugin
+	bool IsValidGameFeaturePlugin(const FString& PluginDescriptorFilename) const;
+
 };
\ No newline at end of file
