Add UGameFeaturesSubsystem::IsGameFeaturePluginLoaded

diff --git a/Engine/Plugins/Experimental/GameFeatures/Source/GameFeatures/Private/GameFeaturesSubsystem.cpp b/Engine/Plugins/Experimental/GameFeatures/Source/GameFeatures/Private/GameFeaturesSubsystem.cpp
--- a/Engine/Plugins/Experimental/GameFeatures/Source/GameFeatures/Private/GameFeaturesSubsystem.cpp
+++ b/Engine/Plugins/Experimental/GameFeatures/Source/GameFeatures/Private/GameFeaturesSubsystem.cpp
@@ -491,20 +491,26 @@ const UGameFeatureData* UGameFeaturesSubsystem::GetGameFeatureDataForRegisteredP
 	return nullptr;
 }
 
+bool UGameFeaturesSubsystem::IsGameFeaturePluginLoaded(const FString& PluginURL) const
+{
+	if (const UGameFeaturePluginStateMachine* StateMachine = FindGameFeaturePluginStateMachine(PluginURL))
+	{
+		return StateMachine->GetCurrentState() >= EGameFeaturePluginState::Loaded;
+	}
+	return false;
+}
+
 void UGameFeaturesSubsystem::LoadGameFeaturePlugin(const FString& PluginURL, const FGameFeaturePluginLoadComplete& CompleteDelegate)
 {
-	if (UGameFeaturePluginStateMachine* StateMachine = FindGameFeaturePluginStateMachine(PluginURL))
+	if (IsGameFeaturePluginLoaded(PluginURL))
 	{
-		if (StateMachine->GetCurrentState() >= EGameFeaturePluginState::Loaded)
-		{
-			FTSTicker::GetCoreTicker().AddTicker(FTickerDelegate::CreateWeakLambda(this, [this, CompleteDelegate](float dts)
-				{
-					CompleteDelegate.ExecuteIfBound(MakeValue());
-					return false;
-				}));
+		FTSTicker::GetCoreTicker().AddTicker(FTickerDelegate::CreateWeakLambda(this, [this, CompleteDelegate](float dts)
+			{
+				CompleteDelegate.ExecuteIfBound(MakeValue());
+				return false;
+			}));
 
-			return; // Early out, we are already loaded.
-		}
+		return; // Early out, we are already loaded.
 	}
 
 	ChangeGameFeatureTargetState(PluginURL, EGameFeatureTargetState::Loaded, CompleteDelegate);
diff --git a/Engine/Plugins/Experimental/GameFeatures/Source/GameFeatures/Public/GameFeaturesSubsystem.h b/Engine/Plugins/Experimental/GameFeatures/Source/GameFeatures/Public/GameFeaturesSubsystem.h
--- a/Engine/Plugins/Experimental/GameFeatures/Source/GameFeatures/Public/GameFeaturesSubsystem.h
+++ b/Engine/Plugins/Experimental/GameFeatures/Source/GameFeatures/Public/GameFeaturesSubsystem.h
@@ -227,6 +227,9 @@ public:
 	/** Returns the game feature data for a registered plugin specified by PluginURL */
 	const UGameFeatureData* GetGameFeatureDataForRegisteredPluginByURL(const FString& PluginURL);
 
+	/** Determines if a plugin is in the Loaded state.*/
+	bool IsGameFeaturePluginLoaded(const FString& PluginURL) const;
+
 	/** Loads a single game feature plugin. */
 	void LoadGameFeaturePlugin(const FString& PluginURL, const FGameFeaturePluginLoadComplete& CompleteDelegate);
 
