Added support for game feature plugins being created in extension directories
- Missing file from fix for UE-142358

diff --git a/Engine/Plugins/Experimental/GameFeatures/Source/GameFeaturesEditor/Private/GameFeaturesEditorModule.cpp b/Engine/Plugins/Experimental/GameFeatures/Source/GameFeaturesEditor/Private/GameFeaturesEditorModule.cpp
--- a/Engine/Plugins/Experimental/GameFeatures/Source/GameFeaturesEditor/Private/GameFeaturesEditorModule.cpp
+++ b/Engine/Plugins/Experimental/GameFeatures/Source/GameFeaturesEditor/Private/GameFeaturesEditorModule.cpp
@@ -14,6 +14,7 @@
 #include "Subsystems/AssetEditorSubsystem.h"
 #include "GameFeatureDataDetailsCustomization.h"
 #include "GameFeaturesEditorSettings.h"
+#include "GameFeaturesSubsystemSettings.h"
 #include "GameFeaturePluginMetadataCustomization.h"
 #include "Logging/MessageLog.h"
 #include "SSettingsEditorCheckoutNotice.h"
@@ -129,12 +130,8 @@ struct FGameFeaturePluginTemplateDescription : public FPluginTemplateDescription
 
 	bool IsRootedInGameFeaturesRoot(const FString& InStr)
 	{
-		const FString DesiredRoot = GetGameFeatureRoot();
-
-		FString TestStr = InStr / TEXT("");
-		FPaths::MakePlatformFilename(TestStr);
-
-		return TestStr.StartsWith(DesiredRoot);
+		const FString ConvertedPath = FPaths::ConvertRelativePathToFull(FPaths::CreateStandardFilename(InStr / TEXT("test.uplugin")));
+		return GetDefault<UGameFeaturesSubsystemSettings>()->IsValidGameFeaturePlugin(ConvertedPath);
 	}
 
 	TSubclassOf<UGameFeatureData> GameFeatureDataClass;
