Updated GameFeatureStateChangeObserver callbacks to more consistently include a PluginURL parameter.

diff --git a/Engine/Plugins/Experimental/CommonConversation/Source/CommonConversationRuntime/Private/ConversationRegistry.cpp b/Engine/Plugins/Experimental/CommonConversation/Source/CommonConversationRuntime/Private/ConversationRegistry.cpp
--- a/Engine/Plugins/Experimental/CommonConversation/Source/CommonConversationRuntime/Private/ConversationRegistry.cpp
+++ b/Engine/Plugins/Experimental/CommonConversation/Source/CommonConversationRuntime/Private/ConversationRegistry.cpp
@@ -178,12 +178,12 @@ void UConversationRegistry::GameFeatureStateModified()
 	}
 }
 
-void UConversationRegistry::OnGameFeatureActivating(const UGameFeatureData* GameFeatureData)
+void UConversationRegistry::OnGameFeatureActivating(const UGameFeatureData* GameFeatureData, const FString& PluginURL)
 {
 	GameFeatureStateModified();
 }
 
-void UConversationRegistry::OnGameFeatureDeactivating(const UGameFeatureData* GameFeatureData, FGameFeatureDeactivatingContext& Context)
+void UConversationRegistry::OnGameFeatureDeactivating(const UGameFeatureData* GameFeatureData, FGameFeatureDeactivatingContext& Context, const FString& PluginURL)
 {
 	GameFeatureStateModified();
 }
diff --git a/Engine/Plugins/Experimental/CommonConversation/Source/CommonConversationRuntime/Public/ConversationRegistry.h b/Engine/Plugins/Experimental/CommonConversation/Source/CommonConversationRuntime/Public/ConversationRegistry.h
--- a/Engine/Plugins/Experimental/CommonConversation/Source/CommonConversationRuntime/Public/ConversationRegistry.h
+++ b/Engine/Plugins/Experimental/CommonConversation/Source/CommonConversationRuntime/Public/ConversationRegistry.h
@@ -107,9 +107,9 @@ private:
 
 	void GameFeatureStateModified();
 
-	virtual void OnGameFeatureActivating(const UGameFeatureData* GameFeatureData) override;
+	virtual void OnGameFeatureActivating(const UGameFeatureData* GameFeatureData, const FString& PluginURL) override;
 
-	virtual void OnGameFeatureDeactivating(const UGameFeatureData* GameFeatureData, FGameFeatureDeactivatingContext& Context) override;
+	virtual void OnGameFeatureDeactivating(const UGameFeatureData* GameFeatureData, FGameFeatureDeactivatingContext& Context, const FString& PluginURL) override;
 
 private:
 	TMap<FSoftObjectPath, TArray<FSoftObjectPath>> RuntimeDependencyGraph;
diff --git a/Engine/Plugins/Experimental/GameFeatures/Source/GameFeatures/Private/GameFeaturePluginStateMachine.cpp b/Engine/Plugins/Experimental/GameFeatures/Source/GameFeatures/Private/GameFeaturePluginStateMachine.cpp
--- a/Engine/Plugins/Experimental/GameFeatures/Source/GameFeatures/Private/GameFeaturePluginStateMachine.cpp
+++ b/Engine/Plugins/Experimental/GameFeatures/Source/GameFeatures/Private/GameFeaturePluginStateMachine.cpp
@@ -1110,7 +1110,7 @@ struct FGameFeaturePluginState_Unregistering : public FGameFeaturePluginState
 		if (StateProperties.GameFeatureData)
 		{
 			const FString PluginName = FPaths::GetBaseFilename(StateProperties.PluginInstalledFilename);
-			UGameFeaturesSubsystem::Get().OnGameFeatureUnregistering(StateProperties.GameFeatureData, PluginName);
+			UGameFeaturesSubsystem::Get().OnGameFeatureUnregistering(StateProperties.GameFeatureData, PluginName, StateProperties.PluginURL);
 			UGameFeaturesSubsystem::Get().UnloadGameFeatureData(StateProperties.GameFeatureData);
 		}
 
@@ -1173,7 +1173,7 @@ struct FGameFeaturePluginState_Registering : public FGameFeaturePluginState
 			StateProperties.GameFeatureData->InitializeBasePluginIniFile(StateProperties.PluginInstalledFilename);
 			StateStatus.SetTransition(EGameFeaturePluginState::Registered);
 
-			UGameFeaturesSubsystem::Get().OnGameFeatureRegistering(StateProperties.GameFeatureData, PluginName);
+			UGameFeaturesSubsystem::Get().OnGameFeatureRegistering(StateProperties.GameFeatureData, PluginName, StateProperties.PluginURL);
 		}
 		else
 		{
@@ -1228,7 +1228,7 @@ struct FGameFeaturePluginState_Loading : public FGameFeaturePluginState
 			BundleHandle->WaitUntilComplete(0.0f, false);
 		}
 
-		UGameFeaturesSubsystem::Get().OnGameFeatureLoading(StateProperties.GameFeatureData);
+		UGameFeaturesSubsystem::Get().OnGameFeatureLoading(StateProperties.GameFeatureData, StateProperties.PluginURL);
 
 		StateStatus.SetTransition(EGameFeaturePluginState::Loaded);
 	}
@@ -1313,7 +1313,7 @@ struct FGameFeaturePluginState_Deactivating : public FGameFeaturePluginState
 			// Deactivate
 			FGameFeatureDeactivatingContext Context(FSimpleDelegate::CreateRaw(this, &FGameFeaturePluginState_Deactivating::OnPauserCompleted));
 			const FString PluginName = FPaths::GetBaseFilename(StateProperties.PluginInstalledFilename);
-			UGameFeaturesSubsystem::Get().OnGameFeatureDeactivating(StateProperties.GameFeatureData, PluginName, Context);
+			UGameFeaturesSubsystem::Get().OnGameFeatureDeactivating(StateProperties.GameFeatureData, PluginName, Context, StateProperties.PluginURL);
 			NumExpectedPausers = Context.NumPausers;
 		}
 
@@ -1344,7 +1344,7 @@ struct FGameFeaturePluginState_Activating : public FGameFeaturePluginState
 		StateProperties.GameFeatureData->InitializeHierarchicalPluginIniFiles(StateProperties.PluginInstalledFilename);
 
 		const FString PluginName = FPaths::GetBaseFilename(StateProperties.PluginInstalledFilename);
-		UGameFeaturesSubsystem::Get().OnGameFeatureActivating(StateProperties.GameFeatureData, PluginName, Context);
+		UGameFeaturesSubsystem::Get().OnGameFeatureActivating(StateProperties.GameFeatureData, PluginName, Context, StateProperties.PluginURL);
 
 		StateStatus.SetTransition(EGameFeaturePluginState::Active);
 	}
diff --git a/Engine/Plugins/Experimental/GameFeatures/Source/GameFeatures/Private/GameFeaturesSubsystem.cpp b/Engine/Plugins/Experimental/GameFeatures/Source/GameFeatures/Private/GameFeaturesSubsystem.cpp
--- a/Engine/Plugins/Experimental/GameFeatures/Source/GameFeatures/Private/GameFeaturesSubsystem.cpp
+++ b/Engine/Plugins/Experimental/GameFeatures/Source/GameFeatures/Private/GameFeaturesSubsystem.cpp
@@ -355,14 +355,14 @@ void UGameFeaturesSubsystem::OnGameFeatureCheckingStatus(const FString& PluginUR
 	}
 }
 
-void UGameFeaturesSubsystem::OnGameFeatureRegistering(const UGameFeatureData* GameFeatureData, const FString& PluginName)
+void UGameFeaturesSubsystem::OnGameFeatureRegistering(const UGameFeatureData* GameFeatureData, const FString& PluginName, const FString& PluginURL)
 {
 	check(GameFeatureData);
 	AddGameFeatureToAssetManager(GameFeatureData, PluginName);
 
 	for (UObject* Observer : Observers)
 	{
-		CastChecked<IGameFeatureStateChangeObserver>(Observer)->OnGameFeatureRegistering(GameFeatureData, PluginName);
+		CastChecked<IGameFeatureStateChangeObserver>(Observer)->OnGameFeatureRegistering(GameFeatureData, PluginName, PluginURL);
 	}
 
 	for (UGameFeatureAction* Action : GameFeatureData->GetActions())
@@ -374,13 +374,13 @@ void UGameFeaturesSubsystem::OnGameFeatureRegistering(const UGameFeatureData* Ga
 	}
 }
 
-void UGameFeaturesSubsystem::OnGameFeatureUnregistering(const UGameFeatureData* GameFeatureData, const FString& PluginName)
+void UGameFeaturesSubsystem::OnGameFeatureUnregistering(const UGameFeatureData* GameFeatureData, const FString& PluginName, const FString& PluginURL)
 {
 	check(GameFeatureData);
 
 	for (UObject* Observer : Observers)
 	{
-		CastChecked<IGameFeatureStateChangeObserver>(Observer)->OnGameFeatureUnregistering(GameFeatureData, PluginName);
+		CastChecked<IGameFeatureStateChangeObserver>(Observer)->OnGameFeatureUnregistering(GameFeatureData, PluginName, PluginURL);
 	}
 
 	for (UGameFeatureAction* Action : GameFeatureData->GetActions())
@@ -392,12 +392,12 @@ void UGameFeaturesSubsystem::OnGameFeatureUnregistering(const UGameFeatureData*
 	}
 }
 
-void UGameFeaturesSubsystem::OnGameFeatureLoading(const UGameFeatureData* GameFeatureData)
+void UGameFeaturesSubsystem::OnGameFeatureLoading(const UGameFeatureData* GameFeatureData, const FString& PluginURL)
 {
 	check(GameFeatureData);
 	for (UObject* Observer : Observers)
 	{
-		CastChecked<IGameFeatureStateChangeObserver>(Observer)->OnGameFeatureLoading(GameFeatureData);
+		CastChecked<IGameFeatureStateChangeObserver>(Observer)->OnGameFeatureLoading(GameFeatureData, PluginURL);
 	}
 
 	for (UGameFeatureAction* Action : GameFeatureData->GetActions())
@@ -409,13 +409,13 @@ void UGameFeaturesSubsystem::OnGameFeatureLoading(const UGameFeatureData* GameFe
 	}
 }
 
-void UGameFeaturesSubsystem::OnGameFeatureActivating(const UGameFeatureData* GameFeatureData, const FString& PluginName, FGameFeatureActivatingContext& Context)
+void UGameFeaturesSubsystem::OnGameFeatureActivating(const UGameFeatureData* GameFeatureData, const FString& PluginName, FGameFeatureActivatingContext& Context, const FString& PluginURL)
 {
 	check(GameFeatureData);
 
 	for (UObject* Observer : Observers)
 	{
-		CastChecked<IGameFeatureStateChangeObserver>(Observer)->OnGameFeatureActivating(GameFeatureData);
+		CastChecked<IGameFeatureStateChangeObserver>(Observer)->OnGameFeatureActivating(GameFeatureData, PluginURL);
 	}
 	
 	for (UGameFeatureAction* Action : GameFeatureData->GetActions())
@@ -427,13 +427,13 @@ void UGameFeaturesSubsystem::OnGameFeatureActivating(const UGameFeatureData* Gam
 	}
 }
 
-void UGameFeaturesSubsystem::OnGameFeatureDeactivating(const UGameFeatureData* GameFeatureData, const FString& PluginName, FGameFeatureDeactivatingContext& Context)
+void UGameFeaturesSubsystem::OnGameFeatureDeactivating(const UGameFeatureData* GameFeatureData, const FString& PluginName, FGameFeatureDeactivatingContext& Context, const FString& PluginURL)
 {
 	check(GameFeatureData);
 
 	for (UObject* Observer : Observers)
 	{
-		CastChecked<IGameFeatureStateChangeObserver>(Observer)->OnGameFeatureDeactivating(GameFeatureData, Context);
+		CastChecked<IGameFeatureStateChangeObserver>(Observer)->OnGameFeatureDeactivating(GameFeatureData, Context, PluginURL);
 	}
 
 	for (UGameFeatureAction* Action : GameFeatureData->GetActions())
@@ -555,7 +555,7 @@ void UGameFeaturesSubsystem::ChangeGameFeatureTargetState(const FString& PluginU
 				// Refire the observer for Activated and do nothing else.
 				for (UObject* Observer : Observers)
 				{
-					CastChecked<IGameFeatureStateChangeObserver>(Observer)->OnGameFeatureActivating(StateMachine->GetGameFeatureDataForActivePlugin());
+					CastChecked<IGameFeatureStateChangeObserver>(Observer)->OnGameFeatureActivating(StateMachine->GetGameFeatureDataForActivePlugin(), PluginURL);
 				}
 
 				FTSTicker::GetCoreTicker().AddTicker(FTickerDelegate::CreateWeakLambda(this, [CompleteDelegate](float)
@@ -1274,4 +1274,4 @@ EGameFeaturePluginState UGameFeaturesSubsystem::ConvertInitialFeatureStateToTarg
 		break;
 	}
 	return InitialState;
-}
+}
\ No newline at end of file
diff --git a/Engine/Plugins/Experimental/GameFeatures/Source/GameFeatures/Public/GameFeatureStateChangeObserver.h b/Engine/Plugins/Experimental/GameFeatures/Source/GameFeatures/Public/GameFeatureStateChangeObserver.h
--- a/Engine/Plugins/Experimental/GameFeatures/Source/GameFeatures/Public/GameFeatureStateChangeObserver.h
+++ b/Engine/Plugins/Experimental/GameFeatures/Source/GameFeatures/Public/GameFeatureStateChangeObserver.h
@@ -35,13 +35,13 @@ public:
 
 	virtual void OnGameFeatureTerminating(const FString& PluginURL) {}
 
-	virtual void OnGameFeatureRegistering(const UGameFeatureData* GameFeatureData, const FString& PluginName) {}
+	virtual void OnGameFeatureRegistering(const UGameFeatureData* GameFeatureData, const FString& PluginName, const FString& PluginURL) {}
 
-	virtual void OnGameFeatureUnregistering(const UGameFeatureData* GameFeatureData, const FString& PluginName) {}
+	virtual void OnGameFeatureUnregistering(const UGameFeatureData* GameFeatureData, const FString& PluginName, const FString& PluginURL) {}
 
-	virtual void OnGameFeatureActivating(const UGameFeatureData* GameFeatureData) {}
+	virtual void OnGameFeatureActivating(const UGameFeatureData* GameFeatureData, const FString& PluginURL) {}
 
-	virtual void OnGameFeatureLoading(const UGameFeatureData* GameFeatureData) {}
+	virtual void OnGameFeatureLoading(const UGameFeatureData* GameFeatureData, const FString& PluginURL) {}
 
-	virtual void OnGameFeatureDeactivating(const UGameFeatureData* GameFeatureData, FGameFeatureDeactivatingContext& Context) {}
+	virtual void OnGameFeatureDeactivating(const UGameFeatureData* GameFeatureData, FGameFeatureDeactivatingContext& Context, const FString& PluginURL) {}
 };
diff --git a/Engine/Plugins/Experimental/GameFeatures/Source/GameFeatures/Public/GameFeaturesSubsystem.h b/Engine/Plugins/Experimental/GameFeatures/Source/GameFeatures/Public/GameFeaturesSubsystem.h
--- a/Engine/Plugins/Experimental/GameFeatures/Source/GameFeatures/Public/GameFeaturesSubsystem.h
+++ b/Engine/Plugins/Experimental/GameFeatures/Source/GameFeatures/Public/GameFeaturesSubsystem.h
@@ -315,19 +315,19 @@ private:
 	void OnGameFeatureCheckingStatus(const FString& PluginURL);
 	friend struct FGameFeaturePluginState_UnknownStatus;
 
-	void OnGameFeatureRegistering(const UGameFeatureData* GameFeatureData, const FString& PluginName);
+	void OnGameFeatureRegistering(const UGameFeatureData* GameFeatureData, const FString& PluginName, const FString& PluginURL);
 	friend struct FGameFeaturePluginState_Registering;
 
-	void OnGameFeatureUnregistering(const UGameFeatureData* GameFeatureData, const FString& PluginName);
+	void OnGameFeatureUnregistering(const UGameFeatureData* GameFeatureData, const FString& PluginName, const FString& PluginURL);
 	friend struct FGameFeaturePluginState_Unregistering;
 
-	void OnGameFeatureActivating(const UGameFeatureData* GameFeatureData, const FString& PluginName, FGameFeatureActivatingContext& Context);
+	void OnGameFeatureActivating(const UGameFeatureData* GameFeatureData, const FString& PluginName, FGameFeatureActivatingContext& Context, const FString& PluginURL);
 	friend struct FGameFeaturePluginState_Activating;
 
-	void OnGameFeatureDeactivating(const UGameFeatureData* GameFeatureData, const FString& PluginName, FGameFeatureDeactivatingContext& Context);
+	void OnGameFeatureDeactivating(const UGameFeatureData* GameFeatureData, const FString& PluginName, FGameFeatureDeactivatingContext& Context, const FString& PluginURL);
 	friend struct FGameFeaturePluginState_Deactivating;
 
-	void OnGameFeatureLoading(const UGameFeatureData* GameFeatureData);
+	void OnGameFeatureLoading(const UGameFeatureData* GameFeatureData, const FString& PluginURL);
 	friend struct FGameFeaturePluginState_Loading;
 
 	void OnAssetManagerCreated();
