Fixing crash on editor shutdown in GameFeaturesEditorModule

diff --git a/Engine/Plugins/Experimental/GameFeatures/Source/GameFeaturesEditor/Private/GameFeaturesEditorModule.cpp b/Engine/Plugins/Experimental/GameFeatures/Source/GameFeaturesEditor/Private/GameFeaturesEditorModule.cpp
--- a/Engine/Plugins/Experimental/GameFeatures/Source/GameFeaturesEditor/Private/GameFeaturesEditorModule.cpp
+++ b/Engine/Plugins/Experimental/GameFeatures/Source/GameFeaturesEditor/Private/GameFeaturesEditorModule.cpp
@@ -157,7 +157,8 @@ class FGameFeaturesEditorModule : public FDefaultModuleImpl
 
 		// Add templates to the new plugin wizard
 		{
-			GetMutableDefault<UGameFeaturesEditorSettings>()->OnSettingChanged().AddRaw(this, &FGameFeaturesEditorModule::OnSettingsChanged);
+			GameFeaturesEditorSettingsWatcher = MakeShared<FGameFeaturesEditorSettingsWatcher>(this);
+			GameFeaturesEditorSettingsWatcher->Init();
 
 			CachePluginTemplates();
 
@@ -186,13 +187,7 @@ class FGameFeaturesEditorModule : public FDefaultModuleImpl
 		// Remove the plugin wizard override
 		if (UObjectInitialized())
 		{
-			if (UClass* GameFeatureClass = UGameFeaturesEditorSettings::StaticClass())
-			{
-				if (UGameFeaturesEditorSettings* GameFeatureEditorSettings = Cast<UGameFeaturesEditorSettings>(GameFeatureClass->GetDefaultObject(false)))
-				{
-					GameFeatureEditorSettings->OnSettingChanged().RemoveAll(this);
-				}
-			}
+			GameFeaturesEditorSettingsWatcher = nullptr;
 
 			IModularFeatures& ModularFeatures = IModularFeatures::Get();
  			ModularFeatures.OnModularFeatureRegistered().RemoveAll(this);
@@ -383,6 +378,34 @@ class FGameFeaturesEditorModule : public FDefaultModuleImpl
 		return nullptr;
 	}
 private:
+
+	struct FGameFeaturesEditorSettingsWatcher : public TSharedFromThis<FGameFeaturesEditorSettingsWatcher>
+	{
+		FGameFeaturesEditorSettingsWatcher(FGameFeaturesEditorModule* InParentModule)
+			: ParentModule(InParentModule)
+		{
+
+		}
+
+		void Init()
+		{
+			GetMutableDefault<UGameFeaturesEditorSettings>()->OnSettingChanged().AddSP(this, &FGameFeaturesEditorSettingsWatcher::OnSettingsChanged);
+		}
+
+		void OnSettingsChanged(UObject* Settings, FPropertyChangedEvent& PropertyChangedEvent)
+		{
+			if (ParentModule != nullptr)
+			{
+				ParentModule->OnSettingsChanged(Settings, PropertyChangedEvent);
+			}
+		}
+	private:
+
+		FGameFeaturesEditorModule* ParentModule;
+	};
+
+	TSharedPtr<FGameFeaturesEditorSettingsWatcher> GameFeaturesEditorSettingsWatcher;
+
 	// Array of Plugin templates populated from GameFeatureDeveloperSettings. Allows projects to
 	//	specify reusable plugin templates for the plugin creation wizard.
 	TArray<TSharedPtr<FGameFeaturePluginTemplateDescription>> PluginTemplates;
