Merging CL#19387281

#UE AddCheats now soft references its extension classes. They can be either sync or async loaded at activation time

diff --git a/Engine/Plugins/Experimental/GameFeatures/Source/GameFeatures/Private/GameFeatureAction_AddCheats.cpp b/Engine/Plugins/Experimental/GameFeatures/Source/GameFeatures/Private/GameFeatureAction_AddCheats.cpp
--- a/Engine/Plugins/Experimental/GameFeatures/Source/GameFeatures/Private/GameFeatureAction_AddCheats.cpp
+++ b/Engine/Plugins/Experimental/GameFeatures/Source/GameFeatures/Private/GameFeatureAction_AddCheats.cpp
@@ -10,6 +10,7 @@
 
 void UGameFeatureAction_AddCheats::OnGameFeatureActivating()
 {
+	bIsActive = true;
 	CheatManagerRegistrationHandle = UCheatManager::RegisterForOnCheatManagerCreated(FOnCheatManagerCreated::FDelegate::CreateUObject(this, &ThisClass::OnCheatManagerCreated));
 }
 
@@ -26,6 +27,7 @@ void UGameFeatureAction_AddCheats::OnGameFeatureDeactivating(FGameFeatureDeactiv
 		}
 	}
 	SpawnedCheatManagers.Empty();
+	bIsActive = false;
 }
 
 #if WITH_EDITOR
@@ -34,9 +36,9 @@ EDataValidationResult UGameFeatureAction_AddCheats::IsDataValid(TArray<FText>& V
 	EDataValidationResult Result = CombineDataValidationResults(Super::IsDataValid(ValidationErrors), EDataValidationResult::Valid);
 
 	int32 EntryIndex = 0;
-	for (const TSubclassOf<UCheatManagerExtension> CheatManagerClass : CheatManagers)
+	for (const TSoftClassPtr<UCheatManagerExtension>& CheatManagerClassPtr : CheatManagers)
 	{
-		if (!CheatManagerClass)
+		if (CheatManagerClassPtr.IsNull())
 		{
 			Result = EDataValidationResult::Invalid;
 			ValidationErrors.Add(FText::Format(LOCTEXT("CheatEntryIsNull", "Null entry at index {0} in CheatManagers"), FText::AsNumber(EntryIndex)));
@@ -59,20 +61,62 @@ void UGameFeatureAction_AddCheats::OnCheatManagerCreated(UCheatManager* CheatMan
 		}
 	}
 
-	for (TSubclassOf<UCheatManagerExtension> CheatManagerClass : CheatManagers)
+	for (const TSoftClassPtr<UCheatManagerExtension>& CheatManagerClassPtr : CheatManagers)
 	{
-		if (CheatManagerClass != nullptr)
+		if (!CheatManagerClassPtr.IsNull())
 		{
-			if ((CheatManagerClass->ClassWithin == nullptr) || CheatManager->IsA(CheatManagerClass->ClassWithin))
+			TSubclassOf<UCheatManagerExtension> CheatManagerClass = CheatManagerClassPtr.Get();
+			if (CheatManagerClass != nullptr)
 			{
-				UCheatManagerExtension* Extension = NewObject<UCheatManagerExtension>(CheatManager, CheatManagerClass);
-				SpawnedCheatManagers.Add(Extension);
-				CheatManager->AddCheatManagerExtension(Extension);
+				// The class is in memory. Spawn now.
+				SpawnCheatManagerExtension(CheatManager, CheatManagerClass);
+			}
+			else if (bLoadCheatManagersAsync)
+			{
+				// The class is not in memory and we want to load async. Start async load now.
+				TWeakObjectPtr<UGameFeatureAction_AddCheats> WeakThis(this);
+				TWeakObjectPtr<UCheatManager> WeakCheatManager(CheatManager);
+				LoadPackageAsync(CheatManagerClassPtr.GetLongPackageName(), FLoadPackageAsyncDelegate::CreateLambda(
+					[WeakThis, WeakCheatManager, CheatManagerClassPtr](const FName& PackageName, UPackage* Package, EAsyncLoadingResult::Type Result)
+				{
+					if (Result == EAsyncLoadingResult::Succeeded)
+					{
+						UGameFeatureAction_AddCheats* StrongThis = WeakThis.Get();
+						UCheatManager* StrongCheatManager = WeakCheatManager.Get();
+						if (StrongThis && StrongThis->bIsActive && StrongCheatManager)
+						{
+							if (TSubclassOf<UCheatManagerExtension> LoadedCheatManagerClass = CheatManagerClassPtr.Get())
+							{
+								StrongThis->SpawnCheatManagerExtension(StrongCheatManager, LoadedCheatManagerClass);
+							}
+						}
+					}
+				}
+				));
+			}
+			else
+			{
+				// The class is not in memory and we want to sync load. Load and spawn immediately.
+				CheatManagerClass = CheatManagerClassPtr.LoadSynchronous();
+				if (CheatManagerClass != nullptr)
+				{
+					SpawnCheatManagerExtension(CheatManager, CheatManagerClass);
+				}
 			}
 		}
 	}
 };
 
+void UGameFeatureAction_AddCheats::SpawnCheatManagerExtension(UCheatManager* CheatManager, const TSubclassOf<UCheatManagerExtension>& CheatManagerClass)
+{
+	if ((CheatManagerClass->ClassWithin == nullptr) || CheatManager->IsA(CheatManagerClass->ClassWithin))
+	{
+		UCheatManagerExtension* Extension = NewObject<UCheatManagerExtension>(CheatManager, CheatManagerClass);
+		SpawnedCheatManagers.Add(Extension);
+		CheatManager->AddCheatManagerExtension(Extension);
+	}
+}
+
 //////////////////////////////////////////////////////////////////////
 
 #undef LOCTEXT_NAMESPACE
diff --git a/Engine/Plugins/Experimental/GameFeatures/Source/GameFeatures/Public/GameFeatureAction_AddCheats.h b/Engine/Plugins/Experimental/GameFeatures/Source/GameFeatures/Public/GameFeatureAction_AddCheats.h
--- a/Engine/Plugins/Experimental/GameFeatures/Source/GameFeatures/Public/GameFeatureAction_AddCheats.h
+++ b/Engine/Plugins/Experimental/GameFeatures/Source/GameFeatures/Public/GameFeatureAction_AddCheats.h
@@ -33,15 +33,21 @@ public:
 
 private:
 	void OnCheatManagerCreated(UCheatManager* CheatManager);
+	void SpawnCheatManagerExtension(UCheatManager* CheatManager, const TSubclassOf<UCheatManagerExtension>& CheatManagerClass);
 
 public:
 	/** Cheat managers to setup for the game feature plugin */
 	UPROPERTY(EditAnywhere, Category="Cheats")
-	TArray<TSubclassOf<UCheatManagerExtension>> CheatManagers;
+	TArray<TSoftClassPtr<UCheatManagerExtension>> CheatManagers;
+
+	UPROPERTY(EditAnywhere, Category="Cheats")
+	bool bLoadCheatManagersAsync;
 
 private:
 	FDelegateHandle CheatManagerRegistrationHandle;
 
 	UPROPERTY(Transient)
 	TArray<TWeakObjectPtr<UCheatManagerExtension>> SpawnedCheatManagers;
+
+	bool bIsActive = false;
 };
