Fixed 'files(s) -> file(s)' typo's in code and comments.

diff --git a/Engine/Plugins/Developer/PerforceSourceControl/Source/PerforceSourceControl/Private/PerforceSourceControlState.cpp b/Engine/Plugins/Developer/PerforceSourceControl/Source/PerforceSourceControl/Private/PerforceSourceControlState.cpp
--- a/Engine/Plugins/Developer/PerforceSourceControl/Source/PerforceSourceControl/Private/PerforceSourceControlState.cpp
+++ b/Engine/Plugins/Developer/PerforceSourceControl/Source/PerforceSourceControl/Private/PerforceSourceControlState.cpp
@@ -175,7 +175,7 @@ FText FPerforceSourceControlState::GetDisplayTooltip() const
 {
 	if (IsConflicted())
 	{
-		return LOCTEXT("Conflicted_Tooltip", "The files(s) have local changes that need to be resolved with changes submitted to the Perforce depot");
+		return LOCTEXT("Conflicted_Tooltip", "The file(s) have local changes that need to be resolved with changes submitted to the Perforce depot");
 	}
 	else if( !IsCurrent() )
 	{
diff --git a/Engine/Plugins/Experimental/GameFeatures/Source/GameFeatures/Private/GameFeaturePluginStateMachine.h b/Engine/Plugins/Experimental/GameFeatures/Source/GameFeatures/Private/GameFeaturePluginStateMachine.h
--- a/Engine/Plugins/Experimental/GameFeatures/Source/GameFeatures/Private/GameFeaturePluginStateMachine.h
+++ b/Engine/Plugins/Experimental/GameFeatures/Source/GameFeatures/Private/GameFeaturePluginStateMachine.h
@@ -146,7 +146,7 @@ enum class EGameFeaturePluginState : uint8
 	ErrorRegistering,			// Error state for Installed -> Registered and Registered -> Installed transitions.
 	WaitingForDependencies,		// Transition state Installed -> Registered. In the process of loading code/content for all dependencies into memory.
 	Unmounting,					// Transition state Registered -> Installed. The content file(s) (i.e. pak file) for the plugin is unmounting.
-	Mounting,					// Transition state Installed -> Registered. The content files(s) (i.e. pak file) for the plugin is getting mounted.
+	Mounting,					// Transition state Installed -> Registered. The content file(s) (i.e. pak file) for the plugin is getting mounted.
 	Unregistering,				// Transition state Registered -> Installed. Cleaning up data gathered in Registering.
 	Registering,				// Transition state Installed -> Registered. Discovering assets in the plugin, but not loading them, except a few for discovery reasons.
 	Registered,					// The assets in the plugin are known, but have not yet been loaded, except a few for discovery reasons.
diff --git a/Engine/Plugins/Experimental/GameFeatures/Source/GameFeaturesEditor/Private/SGameFeatureStateWidget.cpp b/Engine/Plugins/Experimental/GameFeatures/Source/GameFeaturesEditor/Private/SGameFeatureStateWidget.cpp
--- a/Engine/Plugins/Experimental/GameFeatures/Source/GameFeaturesEditor/Private/SGameFeatureStateWidget.cpp
+++ b/Engine/Plugins/Experimental/GameFeatures/Source/GameFeaturesEditor/Private/SGameFeatureStateWidget.cpp
@@ -130,7 +130,7 @@ FText SGameFeatureStateWidget::GetTooltipOfState(EGameFeaturePluginState State)
 	case EGameFeaturePluginState::Unmounting:
 		return LOCTEXT("StateTooltip_Unmounting", "Transition state Registered -> Installed. The content file(s) (i.e. pak file) for the plugin is unmounting.");
 	case EGameFeaturePluginState::Mounting:
-		return LOCTEXT("StateTooltip_Mounting", "Transition state Installed -> Registered. The content files(s) (i.e. pak file) for the plugin is getting mounted.");
+		return LOCTEXT("StateTooltip_Mounting", "Transition state Installed -> Registered. The content file(s) (i.e. pak file) for the plugin is getting mounted.");
 	case EGameFeaturePluginState::Unregistering:
 		return LOCTEXT("StateTooltip_Unregistering", "Transition state Registered -> Installed. Cleaning up data gathered in Registering.");
 	case EGameFeaturePluginState::Registering:
diff --git a/Engine/Source/Developer/IoStoreUtilities/Private/ZenStoreWriter.cpp b/Engine/Source/Developer/IoStoreUtilities/Private/ZenStoreWriter.cpp
--- a/Engine/Source/Developer/IoStoreUtilities/Private/ZenStoreWriter.cpp
+++ b/Engine/Source/Developer/IoStoreUtilities/Private/ZenStoreWriter.cpp
@@ -455,7 +455,7 @@ void FZenStoreWriter::Initialize(const FCookInfo& Info)
 					ZenFileSystemManifest->AddManifestEntry(FileChunkId, MoveTemp(ServerPath), MoveTemp(ClientPath));
 				}
 
-				UE_LOG(LogZenStoreWriter, Display, TEXT("Fetched '%d' files(s) from oplog '%s/%s'"), ZenFileSystemManifest->NumEntries(), *ProjectId, *OplogId);
+				UE_LOG(LogZenStoreWriter, Display, TEXT("Fetched '%d' file(s) from oplog '%s/%s'"), ZenFileSystemManifest->NumEntries(), *ProjectId, *OplogId);
 			}
 			else
 			{
