Fix UGameFeaturesSubsystem::LoadGameFeatureData when the specified asset exists but isn't a GameFeatureData

diff --git a/Engine/Plugins/Experimental/GameFeatures/Source/GameFeatures/Private/GameFeaturesSubsystem.cpp b/Engine/Plugins/Experimental/GameFeatures/Source/GameFeatures/Private/GameFeaturesSubsystem.cpp
--- a/Engine/Plugins/Experimental/GameFeatures/Source/GameFeatures/Private/GameFeaturesSubsystem.cpp
+++ b/Engine/Plugins/Experimental/GameFeatures/Source/GameFeatures/Private/GameFeaturesSubsystem.cpp
@@ -189,7 +189,23 @@ TSharedPtr<FStreamableHandle> UGameFeaturesSubsystem::LoadGameFeatureData(const
 	LocalAssetRegistry.ScanFilesSynchronous({ FPackageName::LongPackageNameToFilename(GameFeaturePackageName, FPackageName::GetAssetPackageExtension()) });
 #endif // WITH_EDITOR
 
-	FAssetData GameFeatureAssetData = LocalAssetRegistry.GetAssetByObjectPath(FName(*GameFeatureToLoad));
+	FAssetData GameFeatureAssetData;
+	{
+		FARFilter ArFilter;
+		ArFilter.ObjectPaths.Add(FName(*GameFeatureToLoad));
+		ArFilter.ClassNames.Add(UGameFeatureData::StaticClass()->GetFName());
+		ArFilter.bRecursiveClasses = true;
+#if !WITH_EDITOR
+		ArFilter.bIncludeOnlyOnDiskAssets = true;
+#endif //if !WITH_EDITOR
+
+		TArray<FAssetData> GameFeatureAssets;
+		if (LocalAssetRegistry.GetAssets(ArFilter, GameFeatureAssets) && !GameFeatureAssets.IsEmpty())
+		{
+			GameFeatureAssetData = MoveTemp(GameFeatureAssets[0]);
+		}
+	}
+
 	if (GameFeatureAssetData.IsValid())
 	{
 		FPrimaryAssetId AssetId = GameFeatureAssetData.GetPrimaryAssetId();
