GameFeaturesSubsystem
- Add IsGameFeaturePluginRegistered
- Remove unimplemented GetPluginFilenameFromPluginName

diff --git a/Engine/Plugins/Experimental/GameFeatures/Source/GameFeatures/Private/GameFeaturesSubsystem.cpp b/Engine/Plugins/Experimental/GameFeatures/Source/GameFeatures/Private/GameFeaturesSubsystem.cpp
--- a/Engine/Plugins/Experimental/GameFeatures/Source/GameFeatures/Private/GameFeaturesSubsystem.cpp
+++ b/Engine/Plugins/Experimental/GameFeatures/Source/GameFeatures/Private/GameFeaturesSubsystem.cpp
@@ -507,6 +507,15 @@ const UGameFeatureData* UGameFeaturesSubsystem::GetGameFeatureDataForRegisteredP
 	return nullptr;
 }
 
+bool UGameFeaturesSubsystem::IsGameFeaturePluginRegistered(const FString& PluginURL) const
+{
+	if (const UGameFeaturePluginStateMachine* StateMachine = FindGameFeaturePluginStateMachine(PluginURL))
+	{
+		return StateMachine->GetCurrentState() >= EGameFeaturePluginState::Registered;
+	}
+	return false;
+}
+
 bool UGameFeaturesSubsystem::IsGameFeaturePluginLoaded(const FString& PluginURL) const
 {
 	if (const UGameFeaturePluginStateMachine* StateMachine = FindGameFeaturePluginStateMachine(PluginURL))
@@ -1290,4 +1299,4 @@ EGameFeaturePluginState UGameFeaturesSubsystem::ConvertInitialFeatureStateToTarg
 		break;
 	}
 	return InitialState;
-}
\ No newline at end of file
+}
diff --git a/Engine/Plugins/Experimental/GameFeatures/Source/GameFeatures/Public/GameFeaturesSubsystem.h b/Engine/Plugins/Experimental/GameFeatures/Source/GameFeatures/Public/GameFeaturesSubsystem.h
--- a/Engine/Plugins/Experimental/GameFeatures/Source/GameFeatures/Public/GameFeaturesSubsystem.h
+++ b/Engine/Plugins/Experimental/GameFeatures/Source/GameFeatures/Public/GameFeaturesSubsystem.h
@@ -227,7 +227,10 @@ public:
 	/** Returns the game feature data for a registered plugin specified by PluginURL */
 	const UGameFeatureData* GetGameFeatureDataForRegisteredPluginByURL(const FString& PluginURL);
 
-	/** Determines if a plugin is in the Loaded state.*/
+	/** Determines if a plugin is in the Registered state (or beyond) */
+	bool IsGameFeaturePluginRegistered(const FString& PluginURL) const;
+
+	/** Determines if a plugin is in the Loaded state (or beyond) */
 	bool IsGameFeaturePluginLoaded(const FString& PluginURL) const;
 
 	/** Loads a single game feature plugin. */
@@ -263,9 +266,6 @@ public:
 
 	/** If the specified plugin is a built-in plugin, return the URL used to identify it. Returns true if the plugin exists, false if it was not found */
 	bool GetPluginURLForBuiltInPluginByName(const FString& PluginName, FString& OutPluginURL) const;
-	
-	/** Get the plugin path from the plugin name */
-	FString GetPluginFilenameFromPluginName(const FString& PluginName);
 
 	/** Get the plugin path from the plugin URL */
 	FString GetPluginFilenameFromPluginURL(const FString& PluginURL) const;
