Add support for preloading data registries and sources in the editor when a game feature is registered instead of waiting for full activation, this will show the valid keys in the picker UI but can result in false positives for accesses in PIE
This is currently opt-in with PreloadInEditor on the game feature action
A more advanced version of this would be to add a completely separate editor preview registration mode, but it is currently difficult to separate editor accesses from PIE ones

diff --git a/Engine/Plugins/Experimental/GameFeatures/Source/GameFeatures/Private/GameFeatureAction_DataRegistry.cpp b/Engine/Plugins/Experimental/GameFeatures/Source/GameFeatures/Private/GameFeatureAction_DataRegistry.cpp
--- a/Engine/Plugins/Experimental/GameFeatures/Source/GameFeatures/Private/GameFeatureAction_DataRegistry.cpp
+++ b/Engine/Plugins/Experimental/GameFeatures/Source/GameFeatures/Private/GameFeatureAction_DataRegistry.cpp
@@ -10,10 +10,41 @@
 
 #define LOCTEXT_NAMESPACE "GameFeatures"
 
+void UGameFeatureAction_DataRegistry::OnGameFeatureRegistering()
+{
+	Super::OnGameFeatureRegistering();
+
+	if (ShouldPreloadAtRegistration())
+	{
+		// TODO: Right now this loads the source for both editor and runtime usage, in the future the preload could be changed to only allow resolves and not full data gets
+
+		UDataRegistrySubsystem* DataRegistrySubsystem = UDataRegistrySubsystem::Get();
+		if (ensure(DataRegistrySubsystem))
+		{
+			for (const TSoftObjectPtr<UDataRegistry>& RegistryToAdd : RegistriesToAdd)
+			{
+				if (!RegistryToAdd.IsNull())
+				{
+					const FSoftObjectPath RegistryPath = RegistryToAdd.ToSoftObjectPath();
+
+					UE_LOG(LogGameFeatures, Log, TEXT("OnGameFeatureRegistering %s: Preloading DataRegistry %s for editor preview"), *GetPathName(), *RegistryPath.ToString())
+					DataRegistrySubsystem->LoadRegistryPath(RegistryPath);
+				}
+			}
+		}
+	}
+}
+
 void UGameFeatureAction_DataRegistry::OnGameFeatureActivating()
 {
 	Super::OnGameFeatureActivating();
 
+	if (ShouldPreloadAtRegistration())
+	{
+		// This already happened at registration
+		return;
+	}
+
 	UDataRegistrySubsystem* DataRegistrySubsystem = UDataRegistrySubsystem::Get();
 	if (ensure(DataRegistrySubsystem))
 	{
@@ -44,9 +75,6 @@ void UGameFeatureAction_DataRegistry::OnGameFeatureActivating()
 						}
 					}
 				}
-
-				// @TODO: If game features get an editor refresh function, this code should be changed to handle it
-				// @TODO: Registries that are late-loaded may not show correct picker UI in editor
 #endif
 
 				DataRegistrySubsystem->LoadRegistryPath(RegistryPath);
@@ -55,10 +83,41 @@ void UGameFeatureAction_DataRegistry::OnGameFeatureActivating()
 	}
 }
 
+void UGameFeatureAction_DataRegistry::OnGameFeatureUnregistering()
+{
+	Super::OnGameFeatureUnregistering();
+
+	if (ShouldPreloadAtRegistration())
+	{
+		UDataRegistrySubsystem* DataRegistrySubsystem = UDataRegistrySubsystem::Get();
+		if (ensure(DataRegistrySubsystem))
+		{
+			for (const TSoftObjectPtr<UDataRegistry>& RegistryToAdd : RegistriesToAdd)
+			{
+				if (!RegistryToAdd.IsNull())
+				{
+					const FSoftObjectPath RegistryPath = RegistryToAdd.ToSoftObjectPath();
+
+					// This should only happen when the user is manually changing phase via the feature editor UI
+					UE_LOG(LogGameFeatures, Log, TEXT("OnGameFeatureUnregistering %s: Temporarily disabling preloaded DataRegistry %s"), *GetPathName(), *RegistryPath.ToString())
+
+					DataRegistrySubsystem->IgnoreRegistryPath(RegistryPath);
+				}
+			}
+		}
+	}
+}
+
 void UGameFeatureAction_DataRegistry::OnGameFeatureDeactivating(FGameFeatureDeactivatingContext& Context)
 {
 	Super::OnGameFeatureDeactivating(Context);
 
+	if (ShouldPreloadAtRegistration())
+	{
+		// This will only happen at unregistration
+		return;
+	}
+
 	UDataRegistrySubsystem* DataRegistrySubsystem = UDataRegistrySubsystem::Get();
 	if (ensure(DataRegistrySubsystem))
 	{
@@ -74,6 +133,12 @@ void UGameFeatureAction_DataRegistry::OnGameFeatureDeactivating(FGameFeatureDeac
 	}
 }
 
+bool UGameFeatureAction_DataRegistry::ShouldPreloadAtRegistration()
+{
+	// We want to preload in interactive editor sessions only
+	return (GIsEditor && !IsRunningCommandlet() && bPreloadInEditor);
+}
+
 #if WITH_EDITORONLY_DATA
 void UGameFeatureAction_DataRegistry::AddAdditionalAssetBundleData(FAssetBundleData& AssetBundleData)
 {
diff --git a/Engine/Plugins/Experimental/GameFeatures/Source/GameFeatures/Private/GameFeatureAction_DataRegistrySource.cpp b/Engine/Plugins/Experimental/GameFeatures/Source/GameFeatures/Private/GameFeatureAction_DataRegistrySource.cpp
--- a/Engine/Plugins/Experimental/GameFeatures/Source/GameFeatures/Private/GameFeatureAction_DataRegistrySource.cpp
+++ b/Engine/Plugins/Experimental/GameFeatures/Source/GameFeatures/Private/GameFeatureAction_DataRegistrySource.cpp
@@ -10,10 +10,51 @@
 
 #define LOCTEXT_NAMESPACE "GameFeatures"
 
+void UGameFeatureAction_DataRegistrySource::OnGameFeatureRegistering()
+{
+	Super::OnGameFeatureRegistering();
+
+	if (ShouldPreloadAtRegistration())
+	{
+		// TODO: Right now this loads the source for both editor and runtime usage, in the future the preload could be changed to only allow resolves and not full data gets
+
+		UDataRegistrySubsystem* DataRegistrySubsystem = UDataRegistrySubsystem::Get();
+		if (ensure(DataRegistrySubsystem))
+		{
+			for (const FDataRegistrySourceToAdd& RegistrySource : SourcesToAdd)
+			{
+				// Don't check the client/server flags as they won't work properly
+				TMap<FDataRegistryType, TArray<FSoftObjectPath>> AssetMap;
+				TArray<FSoftObjectPath>& AssetList = AssetMap.Add(RegistrySource.RegistryToAddTo);
+
+				if (!RegistrySource.DataTableToAdd.IsNull())
+				{
+					UE_LOG(LogGameFeatures, Log, TEXT("OnGameFeatureRegistering %s: Preloading DataRegistrySource %s for editor preview"), *GetPathName(), *RegistrySource.DataTableToAdd.ToString())
+						AssetList.Add(RegistrySource.DataTableToAdd.ToSoftObjectPath());
+				}
+
+				if (!RegistrySource.CurveTableToAdd.IsNull())
+				{
+					UE_LOG(LogGameFeatures, Log, TEXT("OnGameFeatureRegistering %s: Preloading DataRegistrySource %s for editor preview"), *GetPathName(), *RegistrySource.CurveTableToAdd.ToString())
+						AssetList.Add(RegistrySource.CurveTableToAdd.ToSoftObjectPath());
+				}
+
+				DataRegistrySubsystem->PreregisterSpecificAssets(AssetMap, RegistrySource.AssetPriority);
+			}
+		}
+	}
+}
+
 void UGameFeatureAction_DataRegistrySource::OnGameFeatureActivating()
 {
 	Super::OnGameFeatureActivating();
 
+	if (ShouldPreloadAtRegistration())
+	{
+		// This already happened at registration
+		return;
+	}
+
 	UDataRegistrySubsystem* DataRegistrySubsystem = UDataRegistrySubsystem::Get();
 	if (ensure(DataRegistrySubsystem))
 	{
@@ -25,7 +66,7 @@ void UGameFeatureAction_DataRegistrySource::OnGameFeatureActivating()
 		for (const FDataRegistrySourceToAdd& RegistrySource : SourcesToAdd)
 		{
 			const bool bShouldAdd = (bIsServer && RegistrySource.bServerSource) || (bIsClient && RegistrySource.bClientSource);
-			if(bShouldAdd)
+			if (bShouldAdd)
 			{
 				TMap<FDataRegistryType, TArray<FSoftObjectPath>> AssetMap;
 				TArray<FSoftObjectPath>& AssetList = AssetMap.Add(RegistrySource.RegistryToAddTo);
@@ -54,9 +95,6 @@ void UGameFeatureAction_DataRegistrySource::OnGameFeatureActivating()
 						UE_LOG(LogGameFeatures, Log, TEXT("OnGameFeatureActivating %s: DataRegistry source asset %s was not loaded before activation, this may cause a long hitch"), *GetPathName(), *RegistrySource.DataTableToAdd.ToString())
 					}
 				}
-
-				// @TODO: If game features get an editor refresh function, this code should be changed to handle it
-				// @TODO: Registry sources that are late-loaded may not show correct picker UI in editor
 #endif
 
 				// This will either load the sources immediately, or schedule them for load when registries are initialized
@@ -66,10 +104,56 @@ void UGameFeatureAction_DataRegistrySource::OnGameFeatureActivating()
 	}
 }
 
+void UGameFeatureAction_DataRegistrySource::OnGameFeatureUnregistering()
+{
+	Super::OnGameFeatureUnregistering();
+
+	if (ShouldPreloadAtRegistration())
+	{
+		// This should only happen when the user is manually changing phase via the feature editor UI
+		UDataRegistrySubsystem* DataRegistrySubsystem = UDataRegistrySubsystem::Get();
+		if (ensure(DataRegistrySubsystem))
+		{
+			for (const FDataRegistrySourceToAdd& RegistrySource : SourcesToAdd)
+			{
+				if (!RegistrySource.DataTableToAdd.IsNull())
+				{
+					if (!DataRegistrySubsystem->UnregisterSpecificAsset(RegistrySource.RegistryToAddTo, RegistrySource.DataTableToAdd.ToSoftObjectPath()))
+					{
+						UE_LOG(LogGameFeatures, Log, TEXT("OnGameFeatureUnregistering %s: DataRegistry data table %s failed to unregister"), *GetPathName(), *RegistrySource.DataTableToAdd.ToString())
+					}
+					else
+					{
+						UE_LOG(LogGameFeatures, Log, TEXT("OnGameFeatureUnregistering %s: Temporarily disabling preloaded data table %s"), *GetPathName(), *RegistrySource.DataTableToAdd.ToString())
+					}
+				}
+
+				if (!RegistrySource.CurveTableToAdd.IsNull())
+				{
+					if (!DataRegistrySubsystem->UnregisterSpecificAsset(RegistrySource.RegistryToAddTo, RegistrySource.CurveTableToAdd.ToSoftObjectPath()))
+					{
+						UE_LOG(LogGameFeatures, Log, TEXT("OnGameFeatureUnregistering %s: DataRegistry curve table %s failed to unregister"), *GetPathName(), *RegistrySource.CurveTableToAdd.ToString())
+					}
+					else
+					{
+						UE_LOG(LogGameFeatures, Log, TEXT("OnGameFeatureUnregistering %s: Temporarily disabling preloaded curve table %s"), *GetPathName(), *RegistrySource.CurveTableToAdd.ToString())
+					}
+				}
+			}
+		}
+	}
+}
+
 void UGameFeatureAction_DataRegistrySource::OnGameFeatureDeactivating(FGameFeatureDeactivatingContext& Context)
 {
 	Super::OnGameFeatureDeactivating(Context);
 
+	if (ShouldPreloadAtRegistration())
+	{
+		// This will only happen at unregistration
+		return;
+	}
+
 	UDataRegistrySubsystem* DataRegistrySubsystem = UDataRegistrySubsystem::Get();
 	if (ensure(DataRegistrySubsystem))
 	{
@@ -94,6 +178,12 @@ void UGameFeatureAction_DataRegistrySource::OnGameFeatureDeactivating(FGameFeatu
 	}
 }
 
+bool UGameFeatureAction_DataRegistrySource::ShouldPreloadAtRegistration()
+{
+	// We want to preload in interactive editor sessions only
+	return (GIsEditor && !IsRunningCommandlet() && bPreloadInEditor);
+}
+
 #if WITH_EDITORONLY_DATA
 void UGameFeatureAction_DataRegistrySource::AddAdditionalAssetBundleData(FAssetBundleData& AssetBundleData)
 {
@@ -103,7 +193,7 @@ void UGameFeatureAction_DataRegistrySource::AddAdditionalAssetBundleData(FAssetB
 		// Register table assets for preloading, this will only work if the game uses client/server bundle states
 		// @TODO: If another way of preloading data is added, client+server sources should use that instead
 
-		if(!RegistrySource.DataTableToAdd.IsNull())
+		if (!RegistrySource.DataTableToAdd.IsNull())
 		{
 			const FSoftObjectPath DataTableSourcePath = RegistrySource.DataTableToAdd.ToSoftObjectPath();
 			if (RegistrySource.bClientSource)
diff --git a/Engine/Plugins/Experimental/GameFeatures/Source/GameFeatures/Public/GameFeatureAction_DataRegistry.h b/Engine/Plugins/Experimental/GameFeatures/Source/GameFeatures/Public/GameFeatureAction_DataRegistry.h
--- a/Engine/Plugins/Experimental/GameFeatures/Source/GameFeatures/Public/GameFeatureAction_DataRegistry.h
+++ b/Engine/Plugins/Experimental/GameFeatures/Source/GameFeatures/Public/GameFeatureAction_DataRegistry.h
@@ -15,10 +15,14 @@ class UGameFeatureAction_DataRegistry : public UGameFeatureAction
 	GENERATED_BODY()
 
 public:
+	virtual void OnGameFeatureRegistering() override;
+	virtual void OnGameFeatureUnregistering() override;
 	virtual void OnGameFeatureActivating() override;
-
 	virtual void OnGameFeatureDeactivating(FGameFeatureDeactivatingContext& Context) override;
 
+	/** If true, we should load the registry at registration time instead of activation time */
+	virtual bool ShouldPreloadAtRegistration();
+
 #if WITH_EDITORONLY_DATA
 	virtual void AddAdditionalAssetBundleData(FAssetBundleData& AssetBundleData) override;
 #endif
@@ -33,4 +37,8 @@ private:
 	/** List of registry assets to load and initialize */
 	UPROPERTY(EditAnywhere, Category = "Registry Data")
 	TArray<TSoftObjectPtr<UDataRegistry> > RegistriesToAdd;
+
+	/** If true, this will preload the registries when the feature is registered in the editor to support the editor pickers */
+	UPROPERTY(EditAnywhere, Category = "Registry Data")
+	bool bPreloadInEditor;
 };
\ No newline at end of file
diff --git a/Engine/Plugins/Experimental/GameFeatures/Source/GameFeatures/Public/GameFeatureAction_DataRegistrySource.h b/Engine/Plugins/Experimental/GameFeatures/Source/GameFeatures/Public/GameFeatureAction_DataRegistrySource.h
--- a/Engine/Plugins/Experimental/GameFeatures/Source/GameFeatures/Public/GameFeatureAction_DataRegistrySource.h
+++ b/Engine/Plugins/Experimental/GameFeatures/Source/GameFeatures/Public/GameFeatureAction_DataRegistrySource.h
@@ -52,10 +52,14 @@ class UGameFeatureAction_DataRegistrySource : public UGameFeatureAction
 	GENERATED_BODY()
 
 public:
+	virtual void OnGameFeatureRegistering() override;
+	virtual void OnGameFeatureUnregistering() override;
 	virtual void OnGameFeatureActivating() override;
-
 	virtual void OnGameFeatureDeactivating(FGameFeatureDeactivatingContext& Context) override;
 
+	/** If true, we should load the sources at registration time instead of activation time */
+	virtual bool ShouldPreloadAtRegistration();
+
 #if WITH_EDITORONLY_DATA
 	virtual void AddAdditionalAssetBundleData(FAssetBundleData& AssetBundleData) override;
 #endif
@@ -67,6 +71,11 @@ public:
 	//~End of UObject interface
 
 private:
+	/** List of sources to add when this feature is activated */
 	UPROPERTY(EditAnywhere, Category = "Registry Data")
 	TArray<FDataRegistrySourceToAdd> SourcesToAdd;
+
+	/** If true, this will preload the sources when the feature is registered in the editor to support the editor pickers */
+	UPROPERTY(EditAnywhere, Category = "Registry Data")
+	bool bPreloadInEditor;
 };
\ No newline at end of file
