Game feature system reads optional GameFeaturePluginURL json field on uplugin dependencies prior to defaulting to a file protocol URL

diff --git a/Engine/Plugins/Experimental/GameFeatures/Source/GameFeatures/Private/GameFeaturesSubsystem.cpp b/Engine/Plugins/Experimental/GameFeatures/Source/GameFeatures/Private/GameFeaturesSubsystem.cpp
--- a/Engine/Plugins/Experimental/GameFeatures/Source/GameFeatures/Private/GameFeaturesSubsystem.cpp
+++ b/Engine/Plugins/Experimental/GameFeatures/Source/GameFeatures/Private/GameFeaturesSubsystem.cpp
@@ -967,6 +967,7 @@ bool UGameFeaturesSubsystem::GetGameFeaturePluginDetails(const FString& PluginDe
 	{
 		FString NameField = TEXT("Name");
 		FString EnabledField = TEXT("Enabled");
+		FString GameFeaturePluginURLField = TEXT("GameFeaturePluginURL");
 		for (const TSharedPtr<FJsonValue>& PluginElement : *PluginsArray)
 		{
 			if (PluginElement.IsValid())
@@ -981,25 +982,34 @@ bool UGameFeaturesSubsystem::GetGameFeaturePluginDetails(const FString& PluginDe
 
 					if (bElementEnabled)
 					{
-						FString ElementName;
-						ElementObject->TryGetStringField(NameField, ElementName);
-						if (!ElementName.IsEmpty())
+						FString GameFeaturePluginURL;
+						ElementObject->TryGetStringField(GameFeaturePluginURLField, GameFeaturePluginURL);
+						if (!GameFeaturePluginURL.IsEmpty())
 						{
-							TSharedPtr<IPlugin> DependencyPlugin = IPluginManager::Get().FindPlugin(ElementName);
-							if (DependencyPlugin.IsValid())
+							OutPluginDetails.PluginDependencies.Add(GameFeaturePluginURL);
+						}
+						else
+						{
+							FString ElementName;
+							ElementObject->TryGetStringField(NameField, ElementName);
+							if (!ElementName.IsEmpty())
 							{
-								const FString& PluginDependencyDescriptorFilename = DependencyPlugin->GetDescriptorFileName();
-
-								if (!PluginDependencyDescriptorFilename.IsEmpty() &&
-									GetDefault<UGameFeaturesSubsystemSettings>()->IsValidGameFeaturePlugin(FPaths::ConvertRelativePathToFull(PluginDependencyDescriptorFilename)) &&
-									FPaths::FileExists(PluginDependencyDescriptorFilename))
+								TSharedPtr<IPlugin> DependencyPlugin = IPluginManager::Get().FindPlugin(ElementName);
+								if (DependencyPlugin.IsValid())
 								{
-									OutPluginDetails.PluginDependencies.Add(GetPluginURL_FileProtocol(DependencyPlugin->GetDescriptorFileName()));
+									const FString& PluginDependencyDescriptorFilename = DependencyPlugin->GetDescriptorFileName();
+
+									if (!PluginDependencyDescriptorFilename.IsEmpty() &&
+										GetDefault<UGameFeaturesSubsystemSettings>()->IsValidGameFeaturePlugin(FPaths::ConvertRelativePathToFull(PluginDependencyDescriptorFilename)) &&
+										FPaths::FileExists(PluginDependencyDescriptorFilename))
+									{
+										OutPluginDetails.PluginDependencies.Add(GetPluginURL_FileProtocol(DependencyPlugin->GetDescriptorFileName()));
+									}
+								}
+								else
+								{
+									UE_LOG(LogGameFeatures, Display, TEXT("Game feature plugin '%s' has unknown dependency '%s'."), *PluginDescriptorFilename, *ElementName);
 								}
-							}
-							else
-							{
-								UE_LOG(LogGameFeatures, Display, TEXT("Game feature plugin '%s' has unknown dependency '%s'."), *PluginDescriptorFilename, *ElementName);
 							}
 						}
 					}
