Minor GFP stuff:
Reduce boiler-plate when adding GFP states.
EGameFeaturePluginState no longer must be cast to int when accessing state array.
Fix private header GameFeaturePluginStateMachine.h being directly included by other modules.

diff --git a/Engine/Plugins/Experimental/GameFeatures/Source/GameFeatures/Private/GameFeaturePluginStateMachine.cpp b/Engine/Plugins/Experimental/GameFeatures/Source/GameFeatures/Private/GameFeaturePluginStateMachine.cpp
--- a/Engine/Plugins/Experimental/GameFeatures/Source/GameFeatures/Private/GameFeaturePluginStateMachine.cpp
+++ b/Engine/Plugins/Experimental/GameFeatures/Source/GameFeatures/Private/GameFeaturePluginStateMachine.cpp
@@ -39,44 +39,19 @@ namespace UE
 			return Result.HasValue() ? FString(TEXT("Success")) : (FString(TEXT("Failure, ErrorCode=")) + Result.GetError());
 		}
 
+		#define GAME_FEATURE_PLUGIN_STATE_TO_STRING(inEnum, inText) case EGameFeaturePluginState::inEnum: return TEXT(#inEnum);
 		FString ToString(EGameFeaturePluginState InType)
 		{
 			switch (InType)
 			{
-			case EGameFeaturePluginState::Uninitialized: return TEXT("Uninitialized");
-			case EGameFeaturePluginState::Terminal: return TEXT("Terminal");
-			case EGameFeaturePluginState::UnknownStatus: return TEXT("UnknownStatus");
-			case EGameFeaturePluginState::CheckingStatus: return TEXT("CheckingStatus");
-			case EGameFeaturePluginState::ErrorCheckingStatus: return TEXT("ErrorCheckingStatus");
-			case EGameFeaturePluginState::ErrorUnavailable: return TEXT("ErrorUnavailable");
-			case EGameFeaturePluginState::StatusKnown: return TEXT("StatusKnown");
-			case EGameFeaturePluginState::ErrorInstalling: return TEXT("ErrorInstalling");
-			case EGameFeaturePluginState::Uninstalling: return TEXT("Uninstalling");
-			case EGameFeaturePluginState::Downloading: return TEXT("Downloading");
-			case EGameFeaturePluginState::Installed: return TEXT("Installed");
-			case EGameFeaturePluginState::ErrorMounting: return TEXT("ErrorMounting");
-			case EGameFeaturePluginState::ErrorWaitingForDependencies: return TEXT("ErrorWaitingForDependencies");
-			case EGameFeaturePluginState::ErrorRegistering: return TEXT("ErrorRegistering");
-			case EGameFeaturePluginState::Unmounting: return TEXT("Unmounting");
-			case EGameFeaturePluginState::Mounting: return TEXT("Mounting");
-			case EGameFeaturePluginState::WaitingForDependencies: return TEXT("WaitingForDependencies");
-			case EGameFeaturePluginState::Unregistering: return TEXT("Unregistering");
-			case EGameFeaturePluginState::Registering: return TEXT("Registering");
-			case EGameFeaturePluginState::Registered: return TEXT("Registered");
-			case EGameFeaturePluginState::Unloading: return TEXT("Unloading");
-			case EGameFeaturePluginState::Loading: return TEXT("Loading");
-			case EGameFeaturePluginState::Loaded: return TEXT("Loaded");
-			case EGameFeaturePluginState::Deactivating: return TEXT("Deactivating");
-			case EGameFeaturePluginState::Activating: return TEXT("Activating");
-			case EGameFeaturePluginState::Active: return TEXT("Active");
+			GAME_FEATURE_PLUGIN_STATE_LIST(GAME_FEATURE_PLUGIN_STATE_TO_STRING)
 			default:
 				check(0);
 				return FString();
 			}
 		}
+		#undef GAME_FEATURE_PLUGIN_STATE_TO_STRING
 	}
-
-	static_assert((int32)EGameFeaturePluginState::MAX == 26, "");
 }
 
 #define GAME_FEATURE_PLUGIN_PROTOCOL_PREFIX(inEnum, inString) case EGameFeaturePluginProtocol::inEnum: return inString;
@@ -1388,36 +1363,11 @@ void UGameFeaturePluginStateMachine::InitStateMachine(const FString& InPluginURL
 		FGameFeaturePluginRequestUpdateStateMachine::CreateUObject(this, &ThisClass::UpdateStateMachine),
 		FGameFeatureStateProgressUpdate::CreateUObject(this, &ThisClass::UpdateCurrentStateProgress));
 
-	AllStates[(int32)EGameFeaturePluginState::Uninitialized] = MakeUnique<FGameFeaturePluginState_Uninitialized>(StateProperties);
-	AllStates[(int32)EGameFeaturePluginState::Terminal] = MakeUnique<FGameFeaturePluginState_Terminal>(StateProperties);
-	AllStates[(int32)EGameFeaturePluginState::UnknownStatus] = MakeUnique<FGameFeaturePluginState_UnknownStatus>(StateProperties);
-	AllStates[(int32)EGameFeaturePluginState::CheckingStatus] = MakeUnique<FGameFeaturePluginState_CheckingStatus>(StateProperties);
-	AllStates[(int32)EGameFeaturePluginState::ErrorCheckingStatus] = MakeUnique<FGameFeaturePluginState_ErrorCheckingStatus>(StateProperties);
-	AllStates[(int32)EGameFeaturePluginState::ErrorUnavailable] = MakeUnique<FGameFeaturePluginState_ErrorUnavailable>(StateProperties);
-	AllStates[(int32)EGameFeaturePluginState::StatusKnown] = MakeUnique<FGameFeaturePluginState_StatusKnown>(StateProperties);
-	AllStates[(int32)EGameFeaturePluginState::ErrorInstalling] = MakeUnique<FGameFeaturePluginState_ErrorInstalling>(StateProperties);
-	AllStates[(int32)EGameFeaturePluginState::Uninstalling] = MakeUnique<FGameFeaturePluginState_Uninstalling>(StateProperties);
-	AllStates[(int32)EGameFeaturePluginState::Downloading] = MakeUnique<FGameFeaturePluginState_Downloading>(StateProperties);
-	AllStates[(int32)EGameFeaturePluginState::Installed] = MakeUnique<FGameFeaturePluginState_Installed>(StateProperties);
-	AllStates[(int32)EGameFeaturePluginState::ErrorMounting] = MakeUnique<FGameFeaturePluginState_ErrorMounting>(StateProperties);
-	AllStates[(int32)EGameFeaturePluginState::ErrorWaitingForDependencies] = MakeUnique<FGameFeaturePluginState_ErrorWaitingForDependencies>(StateProperties);
-	AllStates[(int32)EGameFeaturePluginState::ErrorRegistering] = MakeUnique<FGameFeaturePluginState_ErrorRegistering>(StateProperties);
-	AllStates[(int32)EGameFeaturePluginState::Unmounting] = MakeUnique<FGameFeaturePluginState_Unmounting>(StateProperties);
-	AllStates[(int32)EGameFeaturePluginState::Mounting] = MakeUnique<FGameFeaturePluginState_Mounting>(StateProperties);
-	AllStates[(int32)EGameFeaturePluginState::WaitingForDependencies] = MakeUnique<FGameFeaturePluginState_WaitingForDependencies>(StateProperties);
-	AllStates[(int32)EGameFeaturePluginState::Unregistering] = MakeUnique<FGameFeaturePluginState_Unregistering>(StateProperties);
-	AllStates[(int32)EGameFeaturePluginState::Registering] = MakeUnique<FGameFeaturePluginState_Registering>(StateProperties);
-	AllStates[(int32)EGameFeaturePluginState::Registered] = MakeUnique<FGameFeaturePluginState_Registered>(StateProperties);
-	AllStates[(int32)EGameFeaturePluginState::Unloading] = MakeUnique<FGameFeaturePluginState_Unloading>(StateProperties);
-	AllStates[(int32)EGameFeaturePluginState::Loading] = MakeUnique<FGameFeaturePluginState_Loading>(StateProperties);
-	AllStates[(int32)EGameFeaturePluginState::Loaded] = MakeUnique<FGameFeaturePluginState_Loaded>(StateProperties);
-	AllStates[(int32)EGameFeaturePluginState::Deactivating] = MakeUnique<FGameFeaturePluginState_Deactivating>(StateProperties);
-	AllStates[(int32)EGameFeaturePluginState::Activating] = MakeUnique<FGameFeaturePluginState_Activating>(StateProperties);
-	AllStates[(int32)EGameFeaturePluginState::Active] = MakeUnique<FGameFeaturePluginState_Active>(StateProperties);
-
-	static_assert((int32)EGameFeaturePluginState::MAX == 26, "");
-
-	AllStates[(int32)CurrentStateInfo.State]->BeginState();
+#define GAME_FEATURE_PLUGIN_STATE_MAKE_STATE(inEnum, inText) AllStates[EGameFeaturePluginState::inEnum] = MakeUnique<FGameFeaturePluginState_##inEnum>(StateProperties);
+	GAME_FEATURE_PLUGIN_STATE_LIST(GAME_FEATURE_PLUGIN_STATE_MAKE_STATE)
+#undef GAME_FEATURE_PLUGIN_STATE_MAKE_STATE
+
+	AllStates[CurrentStateInfo.State]->BeginState();
 }
 
 void UGameFeaturePluginStateMachine::SetDestinationState(EGameFeaturePluginState InDestinationState, FGameFeatureStateTransitionComplete OnFeatureStateTransitionComplete)
@@ -1512,19 +1462,19 @@ UGameFeatureData* UGameFeaturePluginStateMachine::GetGameFeatureDataForRegistere
 bool UGameFeaturePluginStateMachine::IsValidTransitionState(EGameFeaturePluginState InState) const
 {
 	check(InState != EGameFeaturePluginState::MAX);
-	return AllStates[(int32)InState]->GetStateType() == EGameFeaturePluginStateType::Transition;
+	return AllStates[InState]->GetStateType() == EGameFeaturePluginStateType::Transition;
 }
 
 bool UGameFeaturePluginStateMachine::IsValidDestinationState(EGameFeaturePluginState InDestinationState) const
 {
 	check(InDestinationState != EGameFeaturePluginState::MAX);
-	return AllStates[(int32)InDestinationState]->GetStateType() == EGameFeaturePluginStateType::Destination;
+	return AllStates[InDestinationState]->GetStateType() == EGameFeaturePluginStateType::Destination;
 }
 
 bool UGameFeaturePluginStateMachine::IsValidErrorState(EGameFeaturePluginState InDestinationState) const
 {
 	check(InDestinationState != EGameFeaturePluginState::MAX);
-	return AllStates[(int32)InDestinationState]->GetStateType() == EGameFeaturePluginStateType::Error;
+	return AllStates[InDestinationState]->GetStateType() == EGameFeaturePluginStateType::Error;
 }
 
 void UGameFeaturePluginStateMachine::UpdateStateMachine()
@@ -1547,7 +1497,7 @@ void UGameFeaturePluginStateMachine::UpdateStateMachine()
 		bKeepProcessing = false;
 
 		FGameFeaturePluginStateStatus StateStatus;
-		AllStates[(int32)CurrentState]->UpdateState(StateStatus);
+		AllStates[CurrentState]->UpdateState(StateStatus);
 
 		TransitionResult = StateStatus.TransitionResult;
 
@@ -1559,11 +1509,11 @@ void UGameFeaturePluginStateMachine::UpdateStateMachine()
 		if (StateStatus.TransitionToState != EGameFeaturePluginState::Uninitialized)
 		{
 			UE_LOG(LogGameFeatures, Verbose, TEXT("Game feature '%s' transitioning state (%s -> %s)"), *GetGameFeatureName(), *UE::GameFeatures::ToString(CurrentState), *UE::GameFeatures::ToString(StateStatus.TransitionToState));
-			AllStates[(int32)CurrentState]->EndState();
+			AllStates[CurrentState]->EndState();
 			CurrentStateInfo = FGameFeaturePluginStateInfo(StateStatus.TransitionToState);
 			CurrentState = StateStatus.TransitionToState;
 			check(CurrentState != EGameFeaturePluginState::MAX);
-			AllStates[(int32)CurrentState]->BeginState();
+			AllStates[CurrentState]->BeginState();
 			OnStateChangedEvent.Broadcast(this);
 			bKeepProcessing = true;
 		}
diff --git a/Engine/Plugins/Experimental/GameFeatures/Source/GameFeatures/Private/GameFeaturePluginStateMachine.h b/Engine/Plugins/Experimental/GameFeatures/Source/GameFeatures/Private/GameFeaturePluginStateMachine.h
--- a/Engine/Plugins/Experimental/GameFeatures/Source/GameFeatures/Private/GameFeaturePluginStateMachine.h
+++ b/Engine/Plugins/Experimental/GameFeatures/Source/GameFeatures/Private/GameFeaturePluginStateMachine.h
@@ -9,6 +9,7 @@
 #include "UObject/ObjectMacros.h"
 #include "UObject/Object.h"
 #include "GameFeaturePluginOperationResult.h"
+#include "GameFeatureTypes.h"
 #include "GameFeaturePluginStateMachine.generated.h"
 
 class UGameFeatureData;
@@ -127,51 +128,10 @@ Transition states are expected to transition the machine to another state after
                +------------+
 */
 
-/** The states a game feature plugin can be in before fully active */
-enum class EGameFeaturePluginState : uint8
-{
-	Uninitialized,				// Unset. Not yet been set up.
-	Terminal,					// Final State before removal of the state machine
-	UnknownStatus,				// Initialized, but the only thing known is the URL to query status.
-	CheckingStatus,				// Transition state UnknownStatus -> StatusKnown. The status is in the process of being queried.
-	ErrorCheckingStatus,		// Error state for UnknownStatus -> StatusKnown transition.
-	ErrorUnavailable,			// Error state for UnknownStatus -> StatusKnown transition.
-	StatusKnown,				// The plugin's information is known, but no action has taken place yet.
-	ErrorInstalling,			// Error state for Installed -> StatusKnown and StatusKnown -> Installed transitions.
-	Uninstalling,				// Transition state Installed -> StatusKnown. In the process of removing from local storage.
-	Downloading,				// Transition state StatusKnown -> Installed. In the process of adding to local storage.
-	Installed,					// The plugin is in local storage (i.e. it is on the hard drive)
-	ErrorMounting,				// Error state for Installed -> Registered and Registered -> Installed transitions.
-	ErrorWaitingForDependencies,// Error state for Installed -> Registered and Registered -> Installed transitions.
-	ErrorRegistering,			// Error state for Installed -> Registered and Registered -> Installed transitions.
-	WaitingForDependencies,		// Transition state Installed -> Registered. In the process of loading code/content for all dependencies into memory.
-	Unmounting,					// Transition state Registered -> Installed. The content file(s) (i.e. pak file) for the plugin is unmounting.
-	Mounting,					// Transition state Installed -> Registered. The content file(s) (i.e. pak file) for the plugin is getting mounted.
-	Unregistering,				// Transition state Registered -> Installed. Cleaning up data gathered in Registering.
-	Registering,				// Transition state Installed -> Registered. Discovering assets in the plugin, but not loading them, except a few for discovery reasons.
-	Registered,					// The assets in the plugin are known, but have not yet been loaded, except a few for discovery reasons.
-	Unloading,					// Transition state Loaded -> Registered. In the process of removing code/content from memory. 
-	Loading,					// Transition state Registered -> Loaded. In the process of loading code/content into memory.
-	Loaded,						// The plugin is loaded into memory and registered with some game systems but not yet active.
-	Deactivating,				// Transition state Active -> Loaded. Currently unregistering with game systems.
-	Activating,					// Transition state Loaded -> Active. Currently registering plugin code/content with game systems.
-	Active,						// Plugin is fully loaded and active. It is affecting the game.
-
-	MAX
-};
-
-namespace UE
-{
-	namespace GameFeatures
-	{
-		GAMEFEATURES_API FString ToString(EGameFeaturePluginState InType);
-	}
-}
-
 #define GAME_FEATURE_PLUGIN_PROTOCOL_LIST(XPROTO)	\
 	XPROTO(File,			TEXT("file:"))			\
 	XPROTO(InstallBundle,	TEXT("installbundle:"))	\
-	XPROTO(Unknown,			TEXT(""))				\
+	XPROTO(Unknown,			TEXT(""))				
 
 #define GAME_FEATURE_PLUGIN_PROTOCOL_ENUM(inEnum, inString) inEnum,
 enum class EGameFeaturePluginProtocol : uint8
@@ -414,7 +374,7 @@ private:
 	FGameFeaturePluginStateMachineProperties StateProperties;
 
 	/** All state machine state objects */
-	TUniquePtr<FGameFeaturePluginState> AllStates[(int32)EGameFeaturePluginState::MAX];
+	TUniquePtr<FGameFeaturePluginState> AllStates[EGameFeaturePluginState::MAX];
 
 	/** True when we are currently executing UpdateStateMachine, to avoid reentry */
 	bool bInUpdateStateMachine;
diff --git a/Engine/Plugins/Experimental/GameFeatures/Source/GameFeatures/Public/GameFeatureTypes.h b/Engine/Plugins/Experimental/GameFeatures/Source/GameFeatures/Public/GameFeatureTypes.h
new file mode 100644
--- /dev/null
+++ b/Engine/Plugins/Experimental/GameFeatures/Source/GameFeatures/Public/GameFeatureTypes.h
@@ -0,0 +1,45 @@
+// Copyright Epic Games, Inc. All Rights Reserved.
+
+#pragma once
+
+#include "GameFeatureTypesFwd.h"
+
+#define GAME_FEATURE_PLUGIN_STATE_LIST(XSTATE)	\
+	XSTATE(Uninitialized,				NSLOCTEXT("GameFeatures", "UninitializedStateDisplayName", "Uninitialized"))								/* Unset. Not yet been set up. */ \
+	XSTATE(Terminal,					NSLOCTEXT("GameFeatures", "TerminalStateDisplayName", "Terminal"))											/* Final State before removal of the state machine. */ \
+	XSTATE(UnknownStatus,				NSLOCTEXT("GameFeatures", "UnknownStatusStateDisplayName", "UnknownStatus"))								/* Initialized, but the only thing known is the URL to query status. */ \
+	XSTATE(CheckingStatus,				NSLOCTEXT("GameFeatures", "CheckingStatusStateDisplayName", "CheckingStatus"))								/* Transition state UnknownStatus -> StatusKnown. The status is in the process of being queried. */ \
+	XSTATE(ErrorCheckingStatus,			NSLOCTEXT("GameFeatures", "ErrorCheckingStatusStateDisplayName", "ErrorCheckingStatus"))					/* Error state for UnknownStatus -> StatusKnown transition. */ \
+	XSTATE(ErrorUnavailable,			NSLOCTEXT("GameFeatures", "ErrorUnavailableStateDisplayName", "ErrorUnavailable"))							/* Error state for UnknownStatus -> StatusKnown transition. */ \
+	XSTATE(StatusKnown,					NSLOCTEXT("GameFeatures", "StatusKnownStateDisplayName", "StatusKnown"))									/* The plugin's information is known, but no action has taken place yet. */ \
+	XSTATE(ErrorInstalling,				NSLOCTEXT("GameFeatures", "ErrorInstallingStateDisplayName", "ErrorInstalling"))							/* Error state for Installed -> StatusKnown and StatusKnown -> Installed transitions. */ \
+	XSTATE(Uninstalling,				NSLOCTEXT("GameFeatures", "UninstallingStateDisplayName", "Uninstalling"))									/* Transition state Installed -> StatusKnown. In the process of removing from local storage. */ \
+	XSTATE(Downloading,					NSLOCTEXT("GameFeatures", "DownloadingStateDisplayName", "Downloading"))									/* Transition state StatusKnown -> Installed. In the process of adding to local storage. */ \
+	XSTATE(Installed,					NSLOCTEXT("GameFeatures", "InstalledStateDisplayName", "Installed"))										/* The plugin is in local storage (i.e. it is on the hard drive) */ \
+	XSTATE(ErrorMounting,				NSLOCTEXT("GameFeatures", "ErrorMountingStateDisplayName", "ErrorMounting"))								/* Error state for Installed -> Registered and Registered -> Installed transitions. */ \
+	XSTATE(ErrorWaitingForDependencies,	NSLOCTEXT("GameFeatures", "ErrorWaitingForDependenciesStateDisplayName", "ErrorWaitingForDependencies"))	/* Error state for Installed -> Registered and Registered -> Installed transitions. */ \
+	XSTATE(ErrorRegistering,			NSLOCTEXT("GameFeatures", "ErrorRegisteringDisplayName", "ErrorRegistering"))								/* Error state for Installed -> Registered and Registered -> Installed transitions. */ \
+	XSTATE(WaitingForDependencies,		NSLOCTEXT("GameFeatures", "WaitingForDependenciesStateDisplayName", "WaitingForDependencies"))				/* Transition state Installed -> Registered. In the process of loading code/content for all dependencies into memory. */ \
+	XSTATE(Unmounting,					NSLOCTEXT("GameFeatures", "UnmountingStateDisplayName", "Unmounting"))										/* Transition state Registered -> Installed. The content file(s) (i.e. pak file) for the plugin is unmounting. */ \
+	XSTATE(Mounting,					NSLOCTEXT("GameFeatures", "MountingStateDisplayName", "Mounting"))											/* Transition state Installed -> Registered. The content files(s) (i.e. pak file) for the plugin is getting mounted. */ \
+	XSTATE(Unregistering,				NSLOCTEXT("GameFeatures", "UnregisteringStateDisplayName", "Unregistering"))								/* Transition state Registered -> Installed. Cleaning up data gathered in Registering. */ \
+	XSTATE(Registering,					NSLOCTEXT("GameFeatures", "RegisteringStateDisplayName", "Registering"))									/* Transition state Installed -> Registered. Discovering assets in the plugin, but not loading them, except a few for discovery reasons. */ \
+	XSTATE(Registered,					NSLOCTEXT("GameFeatures", "RegisteredStateDisplayName", "Registered"))										/* The assets in the plugin are known, but have not yet been loaded, except a few for discovery reasons. */ \
+	XSTATE(Unloading,					NSLOCTEXT("GameFeatures", "UnloadingStateDisplayName", "Unloading"))										/* Transition state Loaded -> Registered. In the process of removing code/contnet from memory. */ \
+	XSTATE(Loading,						NSLOCTEXT("GameFeatures", "LoadingStateDisplayName", "Loading"))											/* Transition state Registered -> Loaded. In the process of loading code/content into memory. */ \
+	XSTATE(Loaded,						NSLOCTEXT("GameFeatures", "LoadedStateDisplayName", "Loaded"))												/* The plugin is loaded into memory, but not registered with game systems and active. */ \
+	XSTATE(Deactivating,				NSLOCTEXT("GameFeatures", "DeactivatingStateDisplayName", "Deactivating"))									/* Transition state Active -> Loaded. Currently unregistering with game systems. */ \
+	XSTATE(Activating,					NSLOCTEXT("GameFeatures", "ActivatingStateDisplayName", "Activating"))										/* Transition state Loaded -> Active. Currently registering plugin code/content with game systems. */ \
+	XSTATE(Active,						NSLOCTEXT("GameFeatures", "ActiveStateDisplayName", "Active"))												/* Plugin is fully loaded and active. It is affecting the game.  */ 
+
+#define GAME_FEATURE_PLUGIN_STATE_ENUM(inEnum, inText) inEnum,
+namespace GameFeaturePluginStatePrivate
+{
+	enum EGameFeaturePluginState : uint8
+	{
+		GAME_FEATURE_PLUGIN_STATE_LIST(GAME_FEATURE_PLUGIN_STATE_ENUM)
+		MAX
+	};
+}
+typedef GameFeaturePluginStatePrivate::EGameFeaturePluginState EGameFeaturePluginState;
+#undef GAME_FEATURE_PLUGIN_STATE_ENUM
diff --git a/Engine/Plugins/Experimental/GameFeatures/Source/GameFeatures/Public/GameFeatureTypesFwd.h b/Engine/Plugins/Experimental/GameFeatures/Source/GameFeatures/Public/GameFeatureTypesFwd.h
new file mode 100644
--- /dev/null
+++ b/Engine/Plugins/Experimental/GameFeatures/Source/GameFeatures/Public/GameFeatureTypesFwd.h
@@ -0,0 +1,17 @@
+// Copyright Epic Games, Inc. All Rights Reserved.
+
+#pragma once
+
+namespace GameFeaturePluginStatePrivate
+{
+	enum EGameFeaturePluginState : uint8;
+}
+using EGameFeaturePluginState = GameFeaturePluginStatePrivate::EGameFeaturePluginState;
+
+namespace UE
+{
+	namespace GameFeatures
+	{
+		GAMEFEATURES_API FString ToString(EGameFeaturePluginState InType);
+	}
+}
diff --git a/Engine/Plugins/Experimental/GameFeatures/Source/GameFeatures/Public/GameFeaturesSubsystem.h b/Engine/Plugins/Experimental/GameFeatures/Source/GameFeatures/Public/GameFeaturesSubsystem.h
--- a/Engine/Plugins/Experimental/GameFeatures/Source/GameFeatures/Public/GameFeaturesSubsystem.h
+++ b/Engine/Plugins/Experimental/GameFeatures/Source/GameFeatures/Public/GameFeaturesSubsystem.h
@@ -6,6 +6,7 @@
 #include "Subsystems/EngineSubsystem.h"
 #include "GameFeaturePluginOperationResult.h"
 #include "Engine/Engine.h"
+#include "GameFeatureTypesFwd.h"
 
 #include "GameFeaturesSubsystem.generated.h"
 
@@ -15,7 +16,6 @@ struct FStreamableHandle;
 struct FAssetIdentifier;
 class UGameFeatureData;
 class UGameFeaturesProjectPolicies;
-enum class EGameFeaturePluginState : uint8;
 class IPlugin;
 class FJsonObject;
 struct FWorldContext;
diff --git a/Engine/Plugins/Experimental/GameFeatures/Source/GameFeaturesEditor/Private/GameFeatureDataDetailsCustomization.cpp b/Engine/Plugins/Experimental/GameFeatures/Source/GameFeaturesEditor/Private/GameFeatureDataDetailsCustomization.cpp
--- a/Engine/Plugins/Experimental/GameFeatures/Source/GameFeaturesEditor/Private/GameFeatureDataDetailsCustomization.cpp
+++ b/Engine/Plugins/Experimental/GameFeatures/Source/GameFeaturesEditor/Private/GameFeatureDataDetailsCustomization.cpp
@@ -19,7 +19,7 @@
 #include "Misc/Paths.h"
 
 #include "GameFeatureData.h"
-#include "../../GameFeatures/Private/GameFeaturePluginStateMachine.h"
+#include "GameFeatureTypes.h"
 
 #define LOCTEXT_NAMESPACE "GameFeatures"
 
diff --git a/Engine/Plugins/Experimental/GameFeatures/Source/GameFeaturesEditor/Private/GameFeatureDataDetailsCustomization.h b/Engine/Plugins/Experimental/GameFeatures/Source/GameFeaturesEditor/Private/GameFeatureDataDetailsCustomization.h
--- a/Engine/Plugins/Experimental/GameFeatures/Source/GameFeaturesEditor/Private/GameFeatureDataDetailsCustomization.h
+++ b/Engine/Plugins/Experimental/GameFeatures/Source/GameFeaturesEditor/Private/GameFeatureDataDetailsCustomization.h
@@ -6,6 +6,7 @@
 #include "Input/Reply.h"
 #include "IDetailCustomization.h"
 #include "GameFeaturesSubsystem.h"
+#include "GameFeatureTypesFwd.h"
 
 class IDetailLayoutBuilder;
 class SErrorText;
@@ -14,8 +15,6 @@ class IPlugin;
 //////////////////////////////////////////////////////////////////////////
 // FGameFeatureDataDetailsCustomization
 
-enum class EGameFeaturePluginState : uint8;
-
 class FGameFeatureDataDetailsCustomization : public IDetailCustomization
 {
 public:
diff --git a/Engine/Plugins/Experimental/GameFeatures/Source/GameFeaturesEditor/Private/GameFeaturePluginMetadataCustomization.cpp b/Engine/Plugins/Experimental/GameFeatures/Source/GameFeaturesEditor/Private/GameFeaturePluginMetadataCustomization.cpp
--- a/Engine/Plugins/Experimental/GameFeatures/Source/GameFeaturesEditor/Private/GameFeaturePluginMetadataCustomization.cpp
+++ b/Engine/Plugins/Experimental/GameFeatures/Source/GameFeaturesEditor/Private/GameFeaturePluginMetadataCustomization.cpp
@@ -18,7 +18,7 @@
 
 #include "GameFeatureData.h"
 #include "GameFeaturesSubsystem.h"
-#include "../../GameFeatures/Private/GameFeaturePluginStateMachine.h"
+#include "GameFeatureTypes.h"
 
 #define LOCTEXT_NAMESPACE "GameFeatures"
 
diff --git a/Engine/Plugins/Experimental/GameFeatures/Source/GameFeaturesEditor/Private/GameFeaturePluginMetadataCustomization.h b/Engine/Plugins/Experimental/GameFeatures/Source/GameFeaturesEditor/Private/GameFeaturePluginMetadataCustomization.h
--- a/Engine/Plugins/Experimental/GameFeatures/Source/GameFeaturesEditor/Private/GameFeaturePluginMetadataCustomization.h
+++ b/Engine/Plugins/Experimental/GameFeatures/Source/GameFeaturesEditor/Private/GameFeaturePluginMetadataCustomization.h
@@ -6,12 +6,12 @@
 #include "Input/Reply.h"
 #include "IDetailCustomization.h"
 #include "Features/IPluginsEditorFeature.h"
+#include "GameFeatureTypesFwd.h"
 
 class IDetailLayoutBuilder;
 struct FPluginEditingContext;
 class IPlugin;
 struct FPluginDescriptor;
-enum class EGameFeaturePluginState : uint8;
 
 //////////////////////////////////////////////////////////////////////////
 // FGameFeaturePluginMetadataCustomization
diff --git a/Engine/Plugins/Experimental/GameFeatures/Source/GameFeaturesEditor/Private/SGameFeatureStateWidget.cpp b/Engine/Plugins/Experimental/GameFeatures/Source/GameFeaturesEditor/Private/SGameFeatureStateWidget.cpp
--- a/Engine/Plugins/Experimental/GameFeatures/Source/GameFeaturesEditor/Private/SGameFeatureStateWidget.cpp
+++ b/Engine/Plugins/Experimental/GameFeatures/Source/GameFeaturesEditor/Private/SGameFeatureStateWidget.cpp
@@ -5,7 +5,7 @@
 #include "Widgets/SBoxPanel.h"
 #include "Widgets/Input/SSegmentedControl.h"
 #include "GameFeatureData.h"
-#include "../../GameFeatures/Private/GameFeaturePluginStateMachine.h"
+#include "GameFeatureTypes.h"
 
 #define LOCTEXT_NAMESPACE "GameFeatures"
 
@@ -54,37 +54,12 @@ void SGameFeatureStateWidget::Construct(const FArguments& InArgs)
 
 FText SGameFeatureStateWidget::GetDisplayNameOfState(EGameFeaturePluginState State)
 {
-	static_assert((int32)EGameFeaturePluginState::MAX == 26, "");
-
+#define GAME_FEATURE_PLUGIN_STATE_TEXT(inEnum, inText) case EGameFeaturePluginState::inEnum: return inText;
 	switch (State)
 	{
-	case EGameFeaturePluginState::Uninitialized: return LOCTEXT("UninitializedStateDisplayName", "Uninitialized");
-	case EGameFeaturePluginState::Terminal: return LOCTEXT("TerminalDisplayName", "Terminal");
-	case EGameFeaturePluginState::UnknownStatus: return LOCTEXT("UnknownStatusStateDisplayName", "UnknownStatus");
-	case EGameFeaturePluginState::CheckingStatus: return LOCTEXT("CheckingStatusStateDisplayName", "CheckingStatus");
-	case EGameFeaturePluginState::ErrorCheckingStatus: return LOCTEXT("ErrorCheckingStatusDisplayName", "ErrorCheckingStatus");
-	case EGameFeaturePluginState::ErrorUnavailable: return LOCTEXT("ErrorUnavailableDisplayName", "ErrorUnavailable");
-	case EGameFeaturePluginState::StatusKnown: return LOCTEXT("StatusKnownStateDisplayName", "StatusKnown");
-	case EGameFeaturePluginState::ErrorInstalling: return LOCTEXT("ErrorInstallingDisplayName", "ErrorInstalling");
-	case EGameFeaturePluginState::Uninstalling: return LOCTEXT("UninstallingStateDisplayName", "Uninstalling");
-	case EGameFeaturePluginState::Downloading: return LOCTEXT("DownloadingStateDisplayName", "Downloading");
-	case EGameFeaturePluginState::Installed: return LOCTEXT("InstalledStateDisplayName", "Installed");
-	case EGameFeaturePluginState::ErrorMounting: return LOCTEXT("ErrorMountingDisplayName", "ErrorMounting");
-	case EGameFeaturePluginState::ErrorWaitingForDependencies: return LOCTEXT("ErrorWaitingForDependenciesDisplayName", "ErrorWaitingForDependencies");
-	case EGameFeaturePluginState::Unmounting: return LOCTEXT("UnmountingStateDisplayName", "Unmounting");
-	case EGameFeaturePluginState::Mounting: return LOCTEXT("MountingStateDisplayName", "Mounting");
-	case EGameFeaturePluginState::WaitingForDependencies: return LOCTEXT("WaitingForDependenciesStateDisplayName", "WaitingForDependencies");
-	case EGameFeaturePluginState::Unregistering: return LOCTEXT("UnregisteringStateDisplayName", "Unregistering");
-	case EGameFeaturePluginState::Registering: return LOCTEXT("RegisteringStateDisplayName", "Registering");
-	case EGameFeaturePluginState::Registered: return LOCTEXT("RegisteredStateDisplayName", "Registered");
-	case EGameFeaturePluginState::Unloading: return LOCTEXT("UnloadingStateDisplayName", "Unloading");
-	case EGameFeaturePluginState::Loading: return LOCTEXT("LoadingStateDisplayName", "Loading");
-	case EGameFeaturePluginState::Loaded: return LOCTEXT("LoadedStateDisplayName", "Loaded");
-	case EGameFeaturePluginState::Deactivating: return LOCTEXT("DeactivatingStateDisplayName", "Deactivating");
-	case EGameFeaturePluginState::Activating: return LOCTEXT("ActivatingStateDisplayName", "Activating");
-	case EGameFeaturePluginState::Active: return LOCTEXT("ActiveStateDisplayName", "Active");
-	case EGameFeaturePluginState::ErrorRegistering: return LOCTEXT("ErrorRegisteringDisplayName", "ErrorRegistering");
+	GAME_FEATURE_PLUGIN_STATE_LIST(GAME_FEATURE_PLUGIN_STATE_TEXT)
 	}
+#undef GAME_FEATURE_PLUGIN_STATE_TEXT
 
 	const FString FallbackString = UE::GameFeatures::ToString(State);
 	ensureMsgf(false, TEXT("Unknown EGameFeaturePluginState entry %d %s"), (int)State, *FallbackString);
diff --git a/Engine/Plugins/Experimental/GameFeatures/Source/GameFeaturesEditor/Private/SGameFeatureStateWidget.h b/Engine/Plugins/Experimental/GameFeatures/Source/GameFeaturesEditor/Private/SGameFeatureStateWidget.h
--- a/Engine/Plugins/Experimental/GameFeatures/Source/GameFeaturesEditor/Private/SGameFeatureStateWidget.h
+++ b/Engine/Plugins/Experimental/GameFeatures/Source/GameFeaturesEditor/Private/SGameFeatureStateWidget.h
@@ -6,8 +6,7 @@
 #include "Input/Reply.h"
 #include "Widgets/SCompoundWidget.h"
 #include "GameFeaturesSubsystem.h"
-
-enum class EGameFeaturePluginState : uint8;
+#include "GameFeatureTypesFwd.h"
 
 //////////////////////////////////////////////////////////////////////////
 // FGameFeatureDataDetailsCustomization
