GFP State machine
- Only do a GC in Deactivating or Unloading states if the final destination state is the same as the state they are transitioning to, otherwise do all GC in Unregistering state.
- When a transition requires a GC, the transition now waits until the GC is complete.

diff --git a/Engine/Plugins/Experimental/GameFeatures/Source/GameFeatures/Private/GameFeaturePluginStateMachine.cpp b/Engine/Plugins/Experimental/GameFeatures/Source/GameFeatures/Private/GameFeaturePluginStateMachine.cpp
--- a/Engine/Plugins/Experimental/GameFeatures/Source/GameFeatures/Private/GameFeaturePluginStateMachine.cpp
+++ b/Engine/Plugins/Experimental/GameFeatures/Source/GameFeatures/Private/GameFeaturePluginStateMachine.cpp
@@ -23,35 +23,32 @@
 #include "GameFeaturesProjectPolicies.h"
 #include "Containers/Ticker.h"
 
-namespace UE
+namespace UE::GameFeatures
 {
-	namespace GameFeatures
-	{
-		static const FString StateMachineErrorNamespace(TEXT("GameFeaturePlugin.StateMachine."));
+	static const FString StateMachineErrorNamespace(TEXT("GameFeaturePlugin.StateMachine."));
 
-		static int32 ShouldLogMountedFiles = 0;
-		static FAutoConsoleVariableRef CVarShouldLogMountedFiles(TEXT("GameFeaturePlugin.ShouldLogMountedFiles"),
-			ShouldLogMountedFiles,
-			TEXT("Should the newly mounted files be logged."));
+	static int32 ShouldLogMountedFiles = 0;
+	static FAutoConsoleVariableRef CVarShouldLogMountedFiles(TEXT("GameFeaturePlugin.ShouldLogMountedFiles"),
+		ShouldLogMountedFiles,
+		TEXT("Should the newly mounted files be logged."));
 
-		FString ToString(const UE::GameFeatures::FResult& Result)
-		{
-			return Result.HasValue() ? FString(TEXT("Success")) : (FString(TEXT("Failure, ErrorCode=")) + Result.GetError());
-		}
+	FString ToString(const UE::GameFeatures::FResult& Result)
+	{
+		return Result.HasValue() ? FString(TEXT("Success")) : (FString(TEXT("Failure, ErrorCode=")) + Result.GetError());
+	}
 
-		#define GAME_FEATURE_PLUGIN_STATE_TO_STRING(inEnum, inText) case EGameFeaturePluginState::inEnum: return TEXT(#inEnum);
-		FString ToString(EGameFeaturePluginState InType)
+	#define GAME_FEATURE_PLUGIN_STATE_TO_STRING(inEnum, inText) case EGameFeaturePluginState::inEnum: return TEXT(#inEnum);
+	FString ToString(EGameFeaturePluginState InType)
+	{
+		switch (InType)
 		{
-			switch (InType)
-			{
-			GAME_FEATURE_PLUGIN_STATE_LIST(GAME_FEATURE_PLUGIN_STATE_TO_STRING)
-			default:
-				check(0);
-				return FString();
-			}
+		GAME_FEATURE_PLUGIN_STATE_LIST(GAME_FEATURE_PLUGIN_STATE_TO_STRING)
+		default:
+			check(0);
+			return FString();
 		}
-		#undef GAME_FEATURE_PLUGIN_STATE_TO_STRING
 	}
+	#undef GAME_FEATURE_PLUGIN_STATE_TO_STRING
 }
 
 #define GAME_FEATURE_PLUGIN_PROTOCOL_PREFIX(inEnum, inString) case EGameFeaturePluginProtocol::inEnum: return inString;
@@ -69,19 +66,12 @@ const TCHAR* GameFeaturePluginProtocolPrefix(EGameFeaturePluginProtocol Protocol
 
 FGameFeaturePluginState::~FGameFeaturePluginState()
 {
-	if (TickHandle.IsValid())
-	{
-		FTSTicker::GetCoreTicker().RemoveTicker(TickHandle);
-		TickHandle.Reset();
-	}
+	CleanupDeferredUpdateCallbacks();
 }
 
 void FGameFeaturePluginState::UpdateStateMachineDeferred(float Delay /*= 0.0f*/) const
 {
-	if (TickHandle.IsValid())
-	{
-		FTSTicker::GetCoreTicker().RemoveTicker(TickHandle);
-	}
+	CleanupDeferredUpdateCallbacks();
 
 	TickHandle = FTSTicker::GetCoreTicker().AddTicker(FTickerDelegate::CreateLambda([this](float dts) mutable
 	{
@@ -91,6 +81,14 @@ void FGameFeaturePluginState::UpdateStateMachineDeferred(float Delay /*= 0.0f*/)
 	}), Delay);
 }
 
+void FGameFeaturePluginState::GarbageCollectAndUpdateStateMachineDeferred() const
+{
+	GEngine->ForceGarbageCollection(true); // Tick Delayed
+
+	CleanupDeferredUpdateCallbacks();
+	FCoreUObjectDelegates::GetPostGarbageCollect().AddRaw(this, &FGameFeaturePluginState::UpdateStateMachineDeferred, 0.0f);
+}
+
 void FGameFeaturePluginState::UpdateStateMachineImmediate() const
 {
 	StateProperties.OnRequestUpdateStateMachine.ExecuteIfBound();
@@ -101,6 +99,17 @@ void FGameFeaturePluginState::UpdateProgress(float Progress) const
 	StateProperties.OnFeatureStateProgressUpdate.ExecuteIfBound(Progress);
 }
 
+void FGameFeaturePluginState::CleanupDeferredUpdateCallbacks() const
+{
+	if (TickHandle.IsValid())
+	{
+		FTSTicker::GetCoreTicker().RemoveTicker(TickHandle);
+		TickHandle.Reset();
+	}
+
+	FCoreUObjectDelegates::GetPostGarbageCollect().RemoveAll(this);
+}
+
 /*
 =========================================================
   States
@@ -1080,8 +1089,22 @@ struct FGameFeaturePluginState_Unregistering : public FGameFeaturePluginState
 {
 	FGameFeaturePluginState_Unregistering(FGameFeaturePluginStateMachineProperties& InStateProperties) : FGameFeaturePluginState(InStateProperties) {}
 
+	bool bRequestedGC = false;
+
+	virtual void BeginState()
+	{
+		bRequestedGC = false;
+	}
+
 	virtual void UpdateState(FGameFeaturePluginStateStatus& StateStatus) override
 	{
+		if (bRequestedGC)
+		{
+			// @todo after GC, make sure all loaded content is out of memory
+			StateStatus.SetTransition(EGameFeaturePluginState::Unmounting);
+			return;
+		}
+
 		if (StateProperties.GameFeatureData)
 		{
 			const FString PluginName = FPaths::GetBaseFilename(StateProperties.PluginInstalledFilename);
@@ -1091,10 +1114,8 @@ struct FGameFeaturePluginState_Unregistering : public FGameFeaturePluginState
 
 		StateProperties.GameFeatureData = nullptr;
 
-		// @todo GC, then make sure all loaded content is out of memory
-		GEngine->ForceGarbageCollection(true); // this is tick delayed
-
-		StateStatus.SetTransition(EGameFeaturePluginState::Unmounting);
+		bRequestedGC = true;
+		GarbageCollectAndUpdateStateMachineDeferred();
 	}
 };
 
@@ -1179,10 +1200,26 @@ struct FGameFeaturePluginState_Unloading : public FGameFeaturePluginState
 {
 	FGameFeaturePluginState_Unloading(FGameFeaturePluginStateMachineProperties& InStateProperties) : FGameFeaturePluginState(InStateProperties) {}
 
+	bool bRequestedGC = false;
+
+	virtual void BeginState() 
+	{
+		bRequestedGC = false;
+	}
+
 	virtual void UpdateState(FGameFeaturePluginStateStatus& StateStatus) override
 	{
-		// @todo Unload everything that was loaded in Loading, then collect garbage the next frame (since this state may be traversed mid-world tick)
-		//CollectGarbage(GARBAGE_COLLECTION_KEEPFLAGS);
+		// @todo Unload everything that was loaded in Loading
+
+		if (!bRequestedGC && StateProperties.DestinationState == EGameFeaturePluginState::Registered)
+		{
+			// If we aren't going farther than Registered, GC now
+			// otherwise we will defer until closer to our destination state
+			bRequestedGC = true;
+			GarbageCollectAndUpdateStateMachineDeferred();
+			return;
+		}
+
 		StateStatus.SetTransition(EGameFeaturePluginState::Registered);
 	}
 };
@@ -1257,12 +1294,14 @@ struct FGameFeaturePluginState_Deactivating : public FGameFeaturePluginState
 	int32 NumObservedPausers = 0;
 	int32 NumExpectedPausers = 0;
 	bool bInProcessOfDeactivating = false;
+	bool bRequestedGC = false;
 
 	virtual void BeginState() override
 	{
 		NumObservedPausers = 0;
 		NumExpectedPausers = 0;
 		bInProcessOfDeactivating = false;
+		bRequestedGC = false;
 	}
 
 	void OnPauserCompleted()
@@ -1278,6 +1317,13 @@ struct FGameFeaturePluginState_Deactivating : public FGameFeaturePluginState
 
 	virtual void UpdateState(FGameFeaturePluginStateStatus& StateStatus) override
 	{
+		if (bRequestedGC)
+		{
+			check(NumExpectedPausers == NumObservedPausers);
+			StateStatus.SetTransition(EGameFeaturePluginState::Loaded);
+			return;
+		}
+
 		if (!bInProcessOfDeactivating)
 		{
 			// Make sure we won't complete the transition prematurely if someone registers as a pauser but fires immediately
@@ -1294,9 +1340,17 @@ struct FGameFeaturePluginState_Deactivating : public FGameFeaturePluginState
 
 		if (NumExpectedPausers == NumObservedPausers)
 		{
-			GEngine->ForceGarbageCollection(true);
-			StateStatus.SetTransition(EGameFeaturePluginState::Loaded);
-			bInProcessOfDeactivating = false;
+			if(!bRequestedGC && StateProperties.DestinationState == EGameFeaturePluginState::Loaded)
+			{
+				// If we aren't going farther than Loaded, GC now
+				// otherwise we will defer until closer to our destination state
+				bRequestedGC = true;
+				GarbageCollectAndUpdateStateMachineDeferred();
+			}
+			else
+			{
+				StateStatus.SetTransition(EGameFeaturePluginState::Loaded);
+			}
 		}
 		else
 		{
diff --git a/Engine/Plugins/Experimental/GameFeatures/Source/GameFeatures/Private/GameFeaturePluginStateMachine.h b/Engine/Plugins/Experimental/GameFeatures/Source/GameFeatures/Private/GameFeaturePluginStateMachine.h
--- a/Engine/Plugins/Experimental/GameFeatures/Source/GameFeatures/Private/GameFeaturePluginStateMachine.h
+++ b/Engine/Plugins/Experimental/GameFeatures/Source/GameFeatures/Private/GameFeaturePluginStateMachine.h
@@ -128,6 +128,8 @@ Transition states are expected to transition the machine to another state after
                +------------+
 */
 
+// JMarcus TODO: move this stuff into UE::GameFeatures namespace
+
 #define GAME_FEATURE_PLUGIN_PROTOCOL_LIST(XPROTO)	\
 	XPROTO(File,			TEXT("file:"))			\
 	XPROTO(InstallBundle,	TEXT("installbundle:"))	\
@@ -274,11 +276,14 @@ struct FGameFeaturePluginState
 	FGameFeaturePluginStateMachineProperties& StateProperties;
 
 	void UpdateStateMachineDeferred(float Delay = 0.0f) const;
+	void GarbageCollectAndUpdateStateMachineDeferred() const;
 	void UpdateStateMachineImmediate() const;
 
 	void UpdateProgress(float Progress) const;
 
 private:
+	void CleanupDeferredUpdateCallbacks() const;
+
 	mutable FTSTicker::FDelegateHandle TickHandle;
 };
 
@@ -286,13 +291,13 @@ private:
 struct FGameFeaturePluginStateInfo
 {
 	/** The state this info represents */
-	EGameFeaturePluginState State;
+	EGameFeaturePluginState State = EGameFeaturePluginState::Uninitialized;
 
 	/** The progress of this state. Relevant only for transition states. */
-	float Progress;
+	float Progress = 0.0f;
 
-	FGameFeaturePluginStateInfo() : State(EGameFeaturePluginState::Uninitialized), Progress(0.f) {}
-	explicit FGameFeaturePluginStateInfo(EGameFeaturePluginState InState) : State(InState), Progress(0.f) {}
+	FGameFeaturePluginStateInfo() = default;
+	explicit FGameFeaturePluginStateInfo(EGameFeaturePluginState InState) : State(InState) {}
 };
 
 /** A state machine to manage transitioning a game feature plugin from just a URL into a fully loaded and active plugin, including registering its contents with other game systems */
diff --git a/Engine/Plugins/Experimental/GameFeatures/Source/GameFeatures/Public/GameFeaturePluginOperationResult.h b/Engine/Plugins/Experimental/GameFeatures/Source/GameFeatures/Public/GameFeaturePluginOperationResult.h
--- a/Engine/Plugins/Experimental/GameFeatures/Source/GameFeatures/Public/GameFeaturePluginOperationResult.h
+++ b/Engine/Plugins/Experimental/GameFeatures/Source/GameFeatures/Public/GameFeaturePluginOperationResult.h
@@ -5,14 +5,11 @@
 #include "CoreMinimal.h"
 #include "Templates/ValueOrError.h"
 
-namespace UE
+namespace UE::GameFeatures
 {
-	namespace GameFeatures
-	{
-		struct FSuccessType {};
+	struct FSuccessType {};
 
-		using FResult = TValueOrError<FSuccessType, FString>;
+	using FResult = TValueOrError<FSuccessType, FString>;
 
-		GAMEFEATURES_API FString ToString(const FResult& Result);
-	}
+	GAMEFEATURES_API FString ToString(const FResult& Result);
 }
diff --git a/Engine/Plugins/Experimental/GameFeatures/Source/GameFeatures/Public/GameFeatureTypesFwd.h b/Engine/Plugins/Experimental/GameFeatures/Source/GameFeatures/Public/GameFeatureTypesFwd.h
--- a/Engine/Plugins/Experimental/GameFeatures/Source/GameFeatures/Public/GameFeatureTypesFwd.h
+++ b/Engine/Plugins/Experimental/GameFeatures/Source/GameFeatures/Public/GameFeatureTypesFwd.h
@@ -8,10 +8,7 @@ namespace GameFeaturePluginStatePrivate
 }
 using EGameFeaturePluginState = GameFeaturePluginStatePrivate::EGameFeaturePluginState;
 
-namespace UE
+namespace UE::GameFeatures
 {
-	namespace GameFeatures
-	{
-		GAMEFEATURES_API FString ToString(EGameFeaturePluginState InType);
-	}
+	GAMEFEATURES_API FString ToString(EGameFeaturePluginState InType);
 }
