GameFeaturesSubsystem: Include non-builtin plugins to GameFeaturePluginNameToPathMap
Soft-deprecate GetPluginURLForBuiltInPluginByName and replace it by GetPluginURLForRegisteredPluginByName
Will do the actual deprecation in a separate CL

diff --git a/Engine/Plugins/Experimental/GameFeatures/Source/GameFeatures/Private/GameFeaturesSubsystem.cpp b/Engine/Plugins/Experimental/GameFeatures/Source/GameFeatures/Private/GameFeaturesSubsystem.cpp
--- a/Engine/Plugins/Experimental/GameFeatures/Source/GameFeatures/Private/GameFeaturesSubsystem.cpp
+++ b/Engine/Plugins/Experimental/GameFeatures/Source/GameFeatures/Private/GameFeaturesSubsystem.cpp
@@ -373,6 +373,12 @@ void UGameFeaturesSubsystem::OnGameFeatureCheckingStatus(const FString& PluginUR
 
 void UGameFeaturesSubsystem::OnGameFeatureRegistering(const UGameFeatureData* GameFeatureData, const FString& PluginName, const FString& PluginURL)
 {
+	// Map plugin name to plugin URL
+	if (ensure(!GameFeaturePluginNameToPathMap.Contains(PluginName)))
+	{
+		GameFeaturePluginNameToPathMap.Add(PluginName, PluginURL);
+	}
+
 	check(GameFeatureData);
 	AddGameFeatureToAssetManager(GameFeatureData, PluginName);
 
@@ -406,6 +412,9 @@ void UGameFeaturesSubsystem::OnGameFeatureUnregistering(const UGameFeatureData*
 			Action->OnGameFeatureUnregistering();
 		}
 	}
+
+	// Unmap plugin name to plugin URL
+	GameFeaturePluginNameToPathMap.Remove(PluginName);
 }
 
 void UGameFeaturesSubsystem::OnGameFeatureLoading(const UGameFeatureData* GameFeatureData, const FString& PluginURL)
@@ -828,11 +837,6 @@ void UGameFeaturesSubsystem::LoadBuiltInGameFeaturePlugin(const TSharedRef<IPlug
 					{
 						StateMachine->SetDestinationState(DestinationState, FGameFeatureStateTransitionComplete::CreateUObject(this, &ThisClass::LoadGameFeaturePluginComplete));
 					}
-
-					if (!GameFeaturePluginNameToPathMap.Contains(Plugin->GetName()))
-					{
-						GameFeaturePluginNameToPathMap.Add(Plugin->GetName(), PluginURL);
-					}
 				}
 			}
 		}
@@ -852,7 +856,7 @@ void UGameFeaturesSubsystem::LoadBuiltInGameFeaturePlugins(FBuiltInPluginAdditio
 	UAssetManager::Get().PopBulkScanning();
 }
 
-bool UGameFeaturesSubsystem::GetPluginURLForBuiltInPluginByName(const FString& PluginName, FString& OutPluginURL) const
+bool UGameFeaturesSubsystem::GetPluginURLForRegisteredPluginByName(const FString& PluginName, FString& OutPluginURL) const
 {
 	if (const FString* PluginURL = GameFeaturePluginNameToPathMap.Find(PluginName))
 	{
@@ -863,6 +867,11 @@ bool UGameFeaturesSubsystem::GetPluginURLForBuiltInPluginByName(const FString& P
 	return false;
 }
 
+bool UGameFeaturesSubsystem::GetPluginURLForBuiltInPluginByName(const FString& PluginName, FString& OutPluginURL) const
+{
+	return GetPluginURLForRegisteredPluginByName(PluginName, OutPluginURL);
+}
+
 FString UGameFeaturesSubsystem::GetPluginFilenameFromPluginURL(const FString& PluginURL) const
 {
 	FString PluginFilename;
diff --git a/Engine/Plugins/Experimental/GameFeatures/Source/GameFeatures/Public/GameFeaturesSubsystem.h b/Engine/Plugins/Experimental/GameFeatures/Source/GameFeatures/Public/GameFeaturesSubsystem.h
--- a/Engine/Plugins/Experimental/GameFeatures/Source/GameFeatures/Public/GameFeaturesSubsystem.h
+++ b/Engine/Plugins/Experimental/GameFeatures/Source/GameFeatures/Public/GameFeaturesSubsystem.h
@@ -264,6 +264,10 @@ public:
 	void TerminateGameFeaturePlugin(const FString& PluginURL);
 	void TerminateGameFeaturePlugin(const FString& PluginURL, const FGameFeaturePluginUninstallComplete& CompleteDelegate);
 
+	/** If the specified plugin is registered, return the URL used to identify it. Returns true if the plugin exists, false if it was not found */
+	bool GetPluginURLForRegisteredPluginByName(const FString& PluginName, FString& OutPluginURL) const;
+
+	// @TODO: Deprecate GetPluginURLForBuiltInPluginByName
 	/** If the specified plugin is a built-in plugin, return the URL used to identify it. Returns true if the plugin exists, false if it was not found */
 	bool GetPluginURLForBuiltInPluginByName(const FString& PluginName, FString& OutPluginURL) const;
 
diff --git a/Engine/Plugins/Experimental/GameFeatures/Source/GameFeaturesEditor/Private/GameFeatureDataDetailsCustomization.cpp b/Engine/Plugins/Experimental/GameFeatures/Source/GameFeaturesEditor/Private/GameFeatureDataDetailsCustomization.cpp
--- a/Engine/Plugins/Experimental/GameFeatures/Source/GameFeaturesEditor/Private/GameFeatureDataDetailsCustomization.cpp
+++ b/Engine/Plugins/Experimental/GameFeatures/Source/GameFeaturesEditor/Private/GameFeatureDataDetailsCustomization.cpp
@@ -51,7 +51,7 @@ void FGameFeatureDataDetailsCustomization::CustomizeDetails(IDetailLayoutBuilder
 		GameFeature->GetOutermost()->GetName().ParseIntoArray(PathParts, TEXT("/"));
 
 		UGameFeaturesSubsystem& Subsystem = UGameFeaturesSubsystem::Get();
-		Subsystem.GetPluginURLForBuiltInPluginByName(PathParts[0], /*out*/ PluginURL);
+		Subsystem.GetPluginURLForRegisteredPluginByName(PathParts[0], /*out*/ PluginURL);
 		PluginPtr = IPluginManager::Get().FindPlugin(PathParts[0]);
 
 		const float Padding = 8.0f;
