Quick fix for ensure caused by 19749517
Ensure: OnPlaylistDataLoadCompleted failed to find plugin URL
I'd rather not make this change so I'll investigate why it happens with this particular plugin

diff --git a/Engine/Plugins/Experimental/GameFeatures/Source/GameFeatures/Private/GameFeaturesSubsystem.cpp b/Engine/Plugins/Experimental/GameFeatures/Source/GameFeatures/Private/GameFeaturesSubsystem.cpp
--- a/Engine/Plugins/Experimental/GameFeatures/Source/GameFeatures/Private/GameFeaturesSubsystem.cpp
+++ b/Engine/Plugins/Experimental/GameFeatures/Source/GameFeatures/Private/GameFeaturesSubsystem.cpp
@@ -374,7 +374,7 @@ void UGameFeaturesSubsystem::OnGameFeatureCheckingStatus(const FString& PluginUR
 void UGameFeaturesSubsystem::OnGameFeatureRegistering(const UGameFeatureData* GameFeatureData, const FString& PluginName, const FString& PluginURL)
 {
 	// Map plugin name to plugin URL
-	if (ensure(!GameFeaturePluginNameToPathMap.Contains(PluginName)))
+	if (!GameFeaturePluginNameToPathMap.Contains(PluginName))
 	{
 		GameFeaturePluginNameToPathMap.Add(PluginName, PluginURL);
 	}
@@ -837,6 +837,11 @@ void UGameFeaturesSubsystem::LoadBuiltInGameFeaturePlugin(const TSharedRef<IPlug
 					{
 						StateMachine->SetDestinationState(DestinationState, FGameFeatureStateTransitionComplete::CreateUObject(this, &ThisClass::LoadGameFeaturePluginComplete));
 					}
+
+					if (!GameFeaturePluginNameToPathMap.Contains(Plugin->GetName()))
+					{
+						GameFeaturePluginNameToPathMap.Add(Plugin->GetName(), PluginURL);
+					}
 				}
 			}
 		}
