Unload primary game feature assets in FGameFeaturePluginState_Unloading

diff --git a/Engine/Plugins/Experimental/GameFeatures/Source/GameFeatures/Private/GameFeaturePluginStateMachine.cpp b/Engine/Plugins/Experimental/GameFeatures/Source/GameFeatures/Private/GameFeaturePluginStateMachine.cpp
--- a/Engine/Plugins/Experimental/GameFeatures/Source/GameFeatures/Private/GameFeaturePluginStateMachine.cpp
+++ b/Engine/Plugins/Experimental/GameFeatures/Source/GameFeatures/Private/GameFeaturePluginStateMachine.cpp
@@ -1209,9 +1209,15 @@ struct FGameFeaturePluginState_Unloading : public FGameFeaturePluginState
 
 	virtual void UpdateState(FGameFeaturePluginStateStatus& StateStatus) override
 	{
-		// @todo Unload everything that was loaded in Loading
+		if (bRequestedGC)
+		{
+			StateStatus.SetTransition(EGameFeaturePluginState::Registered);
+			return;
+		}
+
+		UnloadGameFeatureBundles(StateProperties.GameFeatureData);
 
-		if (!bRequestedGC && StateProperties.DestinationState == EGameFeaturePluginState::Registered)
+		if (StateProperties.DestinationState == EGameFeaturePluginState::Registered)
 		{
 			// If we aren't going farther than Registered, GC now
 			// otherwise we will defer until closer to our destination state
@@ -1222,6 +1228,31 @@ struct FGameFeaturePluginState_Unloading : public FGameFeaturePluginState
 
 		StateStatus.SetTransition(EGameFeaturePluginState::Registered);
 	}
+
+	void UnloadGameFeatureBundles(const UGameFeatureData* GameFeatureToLoad)
+	{
+		if (GameFeatureToLoad == nullptr)
+		{
+			return;
+		}
+
+		const UGameFeaturesProjectPolicies& Policy = UGameFeaturesSubsystem::Get().GetPolicy();
+
+		// Remove all bundles from feature data and completely unload everything else
+		FPrimaryAssetId GameFeatureAssetId = GameFeatureToLoad->GetPrimaryAssetId();
+		TSharedPtr<FStreamableHandle> Handle = UAssetManager::Get().ChangeBundleStateForPrimaryAssets({ GameFeatureAssetId }, {}, {}, /*bRemoveAllBundles=*/ true);
+		ensureAlways(Handle == nullptr || Handle->HasLoadCompleted()); // Should be no handle since nothing is being loaded
+
+		TArray<FPrimaryAssetId> AssetIds = Policy.GetPreloadAssetListForGameFeature(GameFeatureToLoad, /*bIncludeLoadedAssets=*/true);
+
+		// Don't unload game feature data asset yet, that will happen in FGameFeaturePluginState_Unregistering
+		ensureAlways(AssetIds.RemoveSwap(GameFeatureAssetId, false) == 0);
+
+		if (AssetIds.Num() > 0)
+		{
+			UAssetManager::Get().UnloadPrimaryAssets(AssetIds);
+		}
+	}
 };
 
 struct FGameFeaturePluginState_Loading : public FGameFeaturePluginState
@@ -1250,16 +1281,16 @@ struct FGameFeaturePluginState_Loading : public FGameFeaturePluginState
 	{
 		check(GameFeatureToLoad);
 
-		TArray<FPrimaryAssetId> AssetIdsToLoad;
+		const UGameFeaturesProjectPolicies& Policy = UGameFeaturesSubsystem::Get().GetPolicy<UGameFeaturesProjectPolicies>();
+
+		TArray<FPrimaryAssetId> AssetIdsToLoad = Policy.GetPreloadAssetListForGameFeature(GameFeatureToLoad);
+
 		FPrimaryAssetId GameFeatureAssetId = GameFeatureToLoad->GetPrimaryAssetId();
 		if (GameFeatureAssetId.IsValid())
 		{
 			AssetIdsToLoad.Add(GameFeatureAssetId);
 		}
 
-		UGameFeaturesProjectPolicies& Policy = UGameFeaturesSubsystem::Get().GetPolicy<UGameFeaturesProjectPolicies>();
-		AssetIdsToLoad.Append(Policy.GetPreloadAssetListForGameFeature(GameFeatureToLoad));
-
 		TSharedPtr<FStreamableHandle> RetHandle;
 		if (AssetIdsToLoad.Num() > 0)
 		{
diff --git a/Engine/Plugins/Experimental/GameFeatures/Source/GameFeatures/Private/GameFeaturesSubsystem.cpp b/Engine/Plugins/Experimental/GameFeatures/Source/GameFeatures/Private/GameFeaturesSubsystem.cpp
--- a/Engine/Plugins/Experimental/GameFeatures/Source/GameFeatures/Private/GameFeaturesSubsystem.cpp
+++ b/Engine/Plugins/Experimental/GameFeatures/Source/GameFeatures/Private/GameFeaturesSubsystem.cpp
@@ -199,11 +199,11 @@ TSharedPtr<FStreamableHandle> UGameFeaturesSubsystem::LoadGameFeatureData(const
 		ArFilter.bIncludeOnlyOnDiskAssets = true;
 #endif //if !WITH_EDITOR
 
-		TArray<FAssetData> GameFeatureAssets;
-		if (LocalAssetRegistry.GetAssets(ArFilter, GameFeatureAssets) && !GameFeatureAssets.IsEmpty())
+		LocalAssetRegistry.EnumerateAssets(ArFilter, [&GameFeatureAssetData](const FAssetData& AssetData)
 		{
-			GameFeatureAssetData = MoveTemp(GameFeatureAssets[0]);
-		}
+			GameFeatureAssetData = AssetData;
+			return false;
+		});
 	}
 
 	if (GameFeatureAssetData.IsValid())
diff --git a/Engine/Plugins/Experimental/GameFeatures/Source/GameFeatures/Public/GameFeaturesProjectPolicies.h b/Engine/Plugins/Experimental/GameFeatures/Source/GameFeatures/Public/GameFeaturesProjectPolicies.h
--- a/Engine/Plugins/Experimental/GameFeatures/Source/GameFeatures/Public/GameFeaturesProjectPolicies.h
+++ b/Engine/Plugins/Experimental/GameFeatures/Source/GameFeatures/Public/GameFeaturesProjectPolicies.h
@@ -24,7 +24,7 @@ public:
 	virtual void ShutdownGameFeatureManager() { }
 
 	// Called when a game feature plugin enters the Loading state to determine additional assets to load
-	virtual TArray<FPrimaryAssetId> GetPreloadAssetListForGameFeature(const UGameFeatureData* GameFeatureToLoad) const { return TArray<FPrimaryAssetId>(); }
+	virtual TArray<FPrimaryAssetId> GetPreloadAssetListForGameFeature(const UGameFeatureData* GameFeatureToLoad, bool bIncludeLoadedAssets = false) const { return TArray<FPrimaryAssetId>(); }
 
 	// Returns the bundle state to use for assets returned by GetPreloadAssetListForGameFeature()
 	// See the Asset Manager documentation for more information about asset bundles
