CVar GameFeaturePlugin.VerifyUnload can be used to verify that all assets in a GFP plugin have been unloaded after a call to UnloadGameFeaturePlugin().

diff --git a/Engine/Plugins/Experimental/GameFeatures/Source/GameFeatures/Private/GameFeaturePluginStateMachine.cpp b/Engine/Plugins/Experimental/GameFeatures/Source/GameFeatures/Private/GameFeaturePluginStateMachine.cpp
--- a/Engine/Plugins/Experimental/GameFeatures/Source/GameFeatures/Private/GameFeaturePluginStateMachine.cpp
+++ b/Engine/Plugins/Experimental/GameFeatures/Source/GameFeatures/Private/GameFeaturePluginStateMachine.cpp
@@ -32,6 +32,11 @@ namespace UE::GameFeatures
 		ShouldLogMountedFiles,
 		TEXT("Should the newly mounted files be logged."));
 
+	static TAutoConsoleVariable<bool> CVarVerifyPluginUnload(TEXT("GameFeaturePlugin.VerifyUnload"), 
+		false,
+		TEXT("Verify plugin assets are no longer in memory when unloading."),
+		ECVF_Cheat);
+
 	FString ToString(const UE::GameFeatures::FResult& Result)
 	{
 		return Result.HasValue() ? FString(TEXT("Success")) : (FString(TEXT("Failure, ErrorCode=")) + Result.GetError());
@@ -49,6 +54,59 @@ namespace UE::GameFeatures
 		}
 	}
 	#undef GAME_FEATURE_PLUGIN_STATE_TO_STRING
+
+	// Verify that all assets from this plugin have been unloaded and GC'd	
+	void VerifyAssetsUnloaded(const FString& PluginName, bool bIgnoreGameFeatureData)
+	{
+#if !UE_BUILD_SHIPPING
+		if (!UE::GameFeatures::CVarVerifyPluginUnload.GetValueOnGameThread())
+		{
+			return;
+		}
+
+		FARFilter PluginArFilter;
+		PluginArFilter.PackagePaths.Add(FName(TEXT("/") + PluginName));
+		PluginArFilter.bRecursivePaths = true;
+
+		if (bIgnoreGameFeatureData)
+		{
+			FARFilter RawGameFeatureDataFilter;
+			RawGameFeatureDataFilter.ClassNames.Add(*UGameFeatureData::StaticClass()->GetName());
+			RawGameFeatureDataFilter.bRecursiveClasses = true;
+
+			FARCompiledFilter GameFeatureDataFilter;
+			UAssetManager::Get().GetAssetRegistry().CompileFilter(RawGameFeatureDataFilter, GameFeatureDataFilter);
+
+			UAssetManager::Get().GetAssetRegistry().EnumerateAssets(PluginArFilter, [&PluginName, &GameFeatureDataFilter](const FAssetData& AssetData)
+				{
+					if (UAssetManager::Get().GetAssetRegistry().IsAssetIncludedByFilter(AssetData, GameFeatureDataFilter))
+					{
+						return true;
+					}
+
+					if (AssetData.IsAssetLoaded())
+					{
+						UE_LOG(LogGameFeatures, Error, TEXT("GFP %s failed to unload asset %s!"), *PluginName, *AssetData.GetFullName());
+					}
+
+					return true;
+				});
+		}
+		else
+		{
+			UAssetManager::Get().GetAssetRegistry().EnumerateAssets(PluginArFilter, [&PluginName](const FAssetData& AssetData)
+				{
+					if (AssetData.IsAssetLoaded())
+					{
+						UE_LOG(LogGameFeatures, Error, TEXT("GFP %s failed to unload asset %s!"), *PluginName, *AssetData.GetFullName());
+					}
+
+					return true;
+				});
+		}
+
+#endif // UE_BUILD_SHIPPING
+	}
 }
 
 #define GAME_FEATURE_PLUGIN_PROTOCOL_PREFIX(inEnum, inString) case EGameFeaturePluginProtocol::inEnum: return inString;
@@ -1098,16 +1156,18 @@ struct FGameFeaturePluginState_Unregistering : public FGameFeaturePluginState
 
 	virtual void UpdateState(FGameFeaturePluginStateStatus& StateStatus) override
 	{
+		const FString PluginName = FPaths::GetBaseFilename(StateProperties.PluginInstalledFilename);
+
 		if (bRequestedGC)
 		{
-			// @todo after GC, make sure all loaded content is out of memory
+			UE::GameFeatures::VerifyAssetsUnloaded(PluginName, false);
+
 			StateStatus.SetTransition(EGameFeaturePluginState::Unmounting);
 			return;
 		}
 
 		if (StateProperties.GameFeatureData)
 		{
-			const FString PluginName = FPaths::GetBaseFilename(StateProperties.PluginInstalledFilename);
 			UGameFeaturesSubsystem::Get().OnGameFeatureUnregistering(StateProperties.GameFeatureData, PluginName, StateProperties.PluginURL);
 			UGameFeaturesSubsystem::Get().UnloadGameFeatureData(StateProperties.GameFeatureData);
 		}
@@ -1211,6 +1271,9 @@ struct FGameFeaturePluginState_Unloading : public FGameFeaturePluginState
 	{
 		if (bRequestedGC)
 		{
+			const FString PluginName = FPaths::GetBaseFilename(StateProperties.PluginInstalledFilename);
+			UE::GameFeatures::VerifyAssetsUnloaded(PluginName, true);
+
 			StateStatus.SetTransition(EGameFeaturePluginState::Registered);
 			return;
 		}
diff --git a/Engine/Plugins/Experimental/GameFeatures/Source/GameFeatures/Private/GameFeaturesSubsystem.cpp b/Engine/Plugins/Experimental/GameFeatures/Source/GameFeatures/Private/GameFeaturesSubsystem.cpp
--- a/Engine/Plugins/Experimental/GameFeatures/Source/GameFeatures/Private/GameFeaturesSubsystem.cpp
+++ b/Engine/Plugins/Experimental/GameFeatures/Source/GameFeatures/Private/GameFeaturesSubsystem.cpp
@@ -147,6 +147,27 @@ void UGameFeaturesSubsystem::Initialize(FSubsystemCollectionBase& Collection)
 			}
 		}),
 		ECVF_Cheat);
+
+	IConsoleManager::Get().RegisterConsoleCommand(
+		TEXT("UnloadAndKeepRegisteredGameFeaturePlugin"),
+		TEXT("Unloads a game feature plugin by URL but keeps it registered"),
+		FConsoleCommandWithWorldArgsAndOutputDeviceDelegate::CreateLambda([](const TArray<FString>& Args, UWorld*, FOutputDevice& Ar)
+			{
+				if (Args.Num() > 0)
+				{
+					FString PluginURL;
+					if (!UGameFeaturesSubsystem::Get().GetPluginURLForBuiltInPluginByName(Args[0], /*out*/ PluginURL))
+					{
+						PluginURL = Args[0];
+					}
+					UGameFeaturesSubsystem::Get().UnloadGameFeaturePlugin(PluginURL, true);
+				}
+				else
+				{
+					Ar.Logf(TEXT("Expected a game feature plugin URL as an argument"));
+				}
+			}),
+		ECVF_Cheat);
 }
 
 void UGameFeaturesSubsystem::Deinitialize()
