Add a comment about how asset registry state gets cleaned up

diff --git a/Engine/Plugins/Experimental/GameFeatures/Source/GameFeatures/Private/GameFeaturePluginStateMachine.cpp b/Engine/Plugins/Experimental/GameFeatures/Source/GameFeatures/Private/GameFeaturePluginStateMachine.cpp
--- a/Engine/Plugins/Experimental/GameFeatures/Source/GameFeatures/Private/GameFeaturePluginStateMachine.cpp
+++ b/Engine/Plugins/Experimental/GameFeatures/Source/GameFeatures/Private/GameFeaturePluginStateMachine.cpp
@@ -788,6 +788,8 @@ struct FGameFeaturePluginState_Unmounting : public FGameFeaturePluginState
 		const FString PluginName = FPaths::GetBaseFilename(StateProperties.PluginInstalledFilename);
 		if (TSharedPtr<IPlugin> Plugin = IPluginManager::Get().FindPlugin(PluginName))
 		{
+			// The asset registry listens to FPackageName::OnContentPathDismounted() and 
+			// will automatically cleanup the asset registry state we added for this plugin.
 			IPluginManager::Get().UnmountExplicitlyLoadedPlugin(PluginName, nullptr);
 		}
 
