Manual integration of 19807702 to UE5 GameFeatureSystem: Add to GameFeaturePluginNameToPathMap on StatusKnown state and remove on Terminal state

diff --git a/Engine/Plugins/Experimental/GameFeatures/Source/GameFeatures/Private/GameFeaturePluginStateMachine.cpp b/Engine/Plugins/Experimental/GameFeatures/Source/GameFeatures/Private/GameFeaturePluginStateMachine.cpp
--- a/Engine/Plugins/Experimental/GameFeatures/Source/GameFeatures/Private/GameFeaturePluginStateMachine.cpp
+++ b/Engine/Plugins/Experimental/GameFeatures/Source/GameFeatures/Private/GameFeaturePluginStateMachine.cpp
@@ -215,7 +215,7 @@ struct FGameFeaturePluginState_Terminal : public FDestinationGameFeaturePluginSt
 		checkf(!bEnteredTerminalState, TEXT("Plugin entered terminal state more than once! %s"), *StateProperties.PluginURL);
 		bEnteredTerminalState = true;
 
-		UGameFeaturesSubsystem::Get().OnGameFeatureTerminating(StateProperties.PluginURL);
+		UGameFeaturesSubsystem::Get().OnGameFeatureTerminating(StateProperties.PluginName, StateProperties.PluginURL);
 	}
 };
 
@@ -313,6 +313,7 @@ struct FGameFeaturePluginState_CheckingStatus : public FGameFeaturePluginState
 			return;
 		}
 
+		UGameFeaturesSubsystem::Get().OnGameFeatureStatusKnown(StateProperties.PluginName, StateProperties.PluginURL);
 		StateStatus.SetTransition(EGameFeaturePluginState::StatusKnown);
 	}
 };
@@ -785,12 +786,11 @@ struct FGameFeaturePluginState_Unmounting : public FGameFeaturePluginState
 		PendingBundles.Empty();
 		bUnmounted = false;
 
-		const FString PluginName = FPaths::GetBaseFilename(StateProperties.PluginInstalledFilename);
-		if (TSharedPtr<IPlugin> Plugin = IPluginManager::Get().FindPlugin(PluginName))
+		if (TSharedPtr<IPlugin> Plugin = IPluginManager::Get().FindPlugin(StateProperties.PluginName))
 		{
 			// The asset registry listens to FPackageName::OnContentPathDismounted() and 
 			// will automatically cleanup the asset registry state we added for this plugin.
-			IPluginManager::Get().UnmountExplicitlyLoadedPlugin(PluginName, nullptr);
+			IPluginManager::Get().UnmountExplicitlyLoadedPlugin(StateProperties.PluginName, nullptr);
 		}
 
 		if (StateProperties.GetPluginProtocol() != EGameFeaturePluginProtocol::InstallBundle)
@@ -992,13 +992,12 @@ struct FGameFeaturePluginState_Mounting : public FGameFeaturePluginState
 			return;
 		}
 		
-		const FString PluginName = FPaths::GetBaseFilename(StateProperties.PluginInstalledFilename);
-		IPluginManager::Get().MountExplicitlyLoadedPlugin(PluginName);
+		IPluginManager::Get().MountExplicitlyLoadedPlugin(StateProperties.PluginName);
 
 		// After the new plugin is mounted add the asset registry for that plugin.
 		if (StateProperties.GetPluginProtocol() == EGameFeaturePluginProtocol::InstallBundle)
 		{
-			TSharedPtr<IPlugin> NewlyMountedPlugin = IPluginManager::Get().FindPlugin(PluginName);
+			TSharedPtr<IPlugin> NewlyMountedPlugin = IPluginManager::Get().FindPlugin(StateProperties.PluginName);
 			if (NewlyMountedPlugin.IsValid() && NewlyMountedPlugin->CanContainContent())
 			{
 				TArray<uint8> SerializedAssetData;
@@ -1158,11 +1157,9 @@ struct FGameFeaturePluginState_Unregistering : public FGameFeaturePluginState
 
 	virtual void UpdateState(FGameFeaturePluginStateStatus& StateStatus) override
 	{
-		const FString PluginName = FPaths::GetBaseFilename(StateProperties.PluginInstalledFilename);
-
 		if (bRequestedGC)
 		{
-			UE::GameFeatures::VerifyAssetsUnloaded(PluginName, false);
+			UE::GameFeatures::VerifyAssetsUnloaded(StateProperties.PluginName, false);
 
 			StateStatus.SetTransition(EGameFeaturePluginState::Unmounting);
 			return;
@@ -1170,7 +1167,7 @@ struct FGameFeaturePluginState_Unregistering : public FGameFeaturePluginState
 
 		if (StateProperties.GameFeatureData)
 		{
-			UGameFeaturesSubsystem::Get().OnGameFeatureUnregistering(StateProperties.GameFeatureData, PluginName, StateProperties.PluginURL);
+			UGameFeaturesSubsystem::Get().OnGameFeatureUnregistering(StateProperties.GameFeatureData, StateProperties.PluginName, StateProperties.PluginURL);
 			UGameFeaturesSubsystem::Get().UnloadGameFeatureData(StateProperties.GameFeatureData);
 		}
 
@@ -1187,15 +1184,14 @@ struct FGameFeaturePluginState_Registering : public FGameFeaturePluginState
 
 	virtual void UpdateState(FGameFeaturePluginStateStatus& StateStatus) override
 	{
-		const FString PluginName = FPaths::GetBaseFilename(StateProperties.PluginInstalledFilename);
 		const FString PluginFolder = FPaths::GetPath(StateProperties.PluginInstalledFilename);
 		UGameplayTagsManager::Get().AddTagIniSearchPath(PluginFolder / TEXT("Config") / TEXT("Tags"));
 
-		const FString PreferredGameFeatureDataPath = FString::Printf(TEXT("/%s/%s.%s"), *PluginName, *PluginName, *PluginName);
+		const FString PreferredGameFeatureDataPath = FString::Printf(TEXT("/%s/%s.%s"), *StateProperties.PluginName, *StateProperties.PluginName, *StateProperties.PluginName);
 
-		FString BackupGameFeatureDataPath = TEXT("/") + PluginName + TEXT("/GameFeatureData.GameFeatureData");
+		FString BackupGameFeatureDataPath = TEXT("/") + StateProperties.PluginName + TEXT("/GameFeatureData.GameFeatureData");
 		// Allow game feature location to be overriden globally and from within the plugin
-		FString OverrideIniPathName = PluginName + TEXT("_Override");
+		FString OverrideIniPathName = StateProperties.PluginName + TEXT("_Override");
 		FString OverridePath = GConfig->GetStr(TEXT("GameFeatureData"), *OverrideIniPathName, GGameIni);
 		if (OverridePath.IsEmpty())
 		{
@@ -1227,11 +1223,10 @@ struct FGameFeaturePluginState_Registering : public FGameFeaturePluginState
 
 		if (StateProperties.GameFeatureData)
 		{
-			StateProperties.PluginName = PluginName;
 			StateProperties.GameFeatureData->InitializeBasePluginIniFile(StateProperties.PluginInstalledFilename);
 			StateStatus.SetTransition(EGameFeaturePluginState::Registered);
 
-			UGameFeaturesSubsystem::Get().OnGameFeatureRegistering(StateProperties.GameFeatureData, PluginName, StateProperties.PluginURL);
+			UGameFeaturesSubsystem::Get().OnGameFeatureRegistering(StateProperties.GameFeatureData, StateProperties.PluginName, StateProperties.PluginURL);
 		}
 		else
 		{
@@ -1273,8 +1268,7 @@ struct FGameFeaturePluginState_Unloading : public FGameFeaturePluginState
 	{
 		if (bRequestedGC)
 		{
-			const FString PluginName = FPaths::GetBaseFilename(StateProperties.PluginInstalledFilename);
-			UE::GameFeatures::VerifyAssetsUnloaded(PluginName, true);
+			UE::GameFeatures::VerifyAssetsUnloaded(StateProperties.PluginName, true);
 
 			StateStatus.SetTransition(EGameFeaturePluginState::Registered);
 			return;
@@ -1429,8 +1423,7 @@ struct FGameFeaturePluginState_Deactivating : public FGameFeaturePluginState
 
 			// Deactivate
 			FGameFeatureDeactivatingContext Context(FSimpleDelegate::CreateRaw(this, &FGameFeaturePluginState_Deactivating::OnPauserCompleted));
-			const FString PluginName = FPaths::GetBaseFilename(StateProperties.PluginInstalledFilename);
-			UGameFeaturesSubsystem::Get().OnGameFeatureDeactivating(StateProperties.GameFeatureData, PluginName, Context, StateProperties.PluginURL);
+			UGameFeaturesSubsystem::Get().OnGameFeatureDeactivating(StateProperties.GameFeatureData, StateProperties.PluginName, Context, StateProperties.PluginURL);
 			NumExpectedPausers = Context.NumPausers;
 		}
 
@@ -1468,8 +1461,7 @@ struct FGameFeaturePluginState_Activating : public FGameFeaturePluginState
 
 		StateProperties.GameFeatureData->InitializeHierarchicalPluginIniFiles(StateProperties.PluginInstalledFilename);
 
-		const FString PluginName = FPaths::GetBaseFilename(StateProperties.PluginInstalledFilename);
-		UGameFeaturesSubsystem::Get().OnGameFeatureActivating(StateProperties.GameFeatureData, PluginName, Context, StateProperties.PluginURL);
+		UGameFeaturesSubsystem::Get().OnGameFeatureActivating(StateProperties.GameFeatureData, StateProperties.PluginName, Context, StateProperties.PluginURL);
 
 		StateStatus.SetTransition(EGameFeaturePluginState::Active);
 	}
@@ -1768,6 +1760,8 @@ bool FGameFeaturePluginStateMachineProperties::ParseURL()
 		return false;
 	}
 
+	PluginName = FPaths::GetBaseFilename(PluginInstalledFilename);
+
 	if (PluginInstalledFilename.IsEmpty() || !PluginInstalledFilename.EndsWith(TEXT(".uplugin")))
 	{
 		ensureMsgf(false, TEXT("PluginInstalledFilename must have a uplugin extension. PluginURL: %s"), *PluginURL);
diff --git a/Engine/Plugins/Experimental/GameFeatures/Source/GameFeatures/Private/GameFeaturesSubsystem.cpp b/Engine/Plugins/Experimental/GameFeatures/Source/GameFeatures/Private/GameFeaturesSubsystem.cpp
--- a/Engine/Plugins/Experimental/GameFeatures/Source/GameFeatures/Private/GameFeaturesSubsystem.cpp
+++ b/Engine/Plugins/Experimental/GameFeatures/Source/GameFeatures/Private/GameFeaturesSubsystem.cpp
@@ -376,12 +376,18 @@ FString UGameFeaturesSubsystem::GetPluginURL_InstallBundleProtocol(const FString
 	return GetPluginURL_InstallBundleProtocol(PluginName, MakeArrayView(&BundleName, 1));
 }
 
-void UGameFeaturesSubsystem::OnGameFeatureTerminating(const FString& PluginURL)
+void UGameFeaturesSubsystem::OnGameFeatureTerminating(const FString& PluginName, const FString& PluginURL)
 {
 	for (UObject* Observer : Observers)
 	{
 		CastChecked<IGameFeatureStateChangeObserver>(Observer)->OnGameFeatureTerminating(PluginURL);
 	}
+
+	if (!PluginName.IsEmpty())
+	{
+		// Unmap plugin name to plugin URL
+		GameFeaturePluginNameToPathMap.Remove(PluginName);
+	}
 }
 
 void UGameFeaturesSubsystem::OnGameFeatureCheckingStatus(const FString& PluginURL)
@@ -392,14 +398,17 @@ void UGameFeaturesSubsystem::OnGameFeatureCheckingStatus(const FString& PluginUR
 	}
 }
 
-void UGameFeaturesSubsystem::OnGameFeatureRegistering(const UGameFeatureData* GameFeatureData, const FString& PluginName, const FString& PluginURL)
+void UGameFeaturesSubsystem::OnGameFeatureStatusKnown(const FString& PluginName, const FString& PluginURL)
 {
 	// Map plugin name to plugin URL
-	if (!GameFeaturePluginNameToPathMap.Contains(PluginName))
+	if (ensure(!GameFeaturePluginNameToPathMap.Contains(PluginName)))
 	{
 		GameFeaturePluginNameToPathMap.Add(PluginName, PluginURL);
 	}
+}
 
+void UGameFeaturesSubsystem::OnGameFeatureRegistering(const UGameFeatureData* GameFeatureData, const FString& PluginName, const FString& PluginURL)
+{
 	check(GameFeatureData);
 	AddGameFeatureToAssetManager(GameFeatureData, PluginName);
 
@@ -433,9 +442,6 @@ void UGameFeaturesSubsystem::OnGameFeatureUnregistering(const UGameFeatureData*
 			Action->OnGameFeatureUnregistering();
 		}
 	}
-
-	// Unmap plugin name to plugin URL
-	GameFeaturePluginNameToPathMap.Remove(PluginName);
 }
 
 void UGameFeaturesSubsystem::OnGameFeatureLoading(const UGameFeatureData* GameFeatureData, const FString& PluginURL)
@@ -858,11 +864,6 @@ void UGameFeaturesSubsystem::LoadBuiltInGameFeaturePlugin(const TSharedRef<IPlug
 					{
 						StateMachine->SetDestinationState(DestinationState, FGameFeatureStateTransitionComplete::CreateUObject(this, &ThisClass::LoadGameFeaturePluginComplete));
 					}
-
-					if (!GameFeaturePluginNameToPathMap.Contains(Plugin->GetName()))
-					{
-						GameFeaturePluginNameToPathMap.Add(Plugin->GetName(), PluginURL);
-					}
 				}
 			}
 		}
@@ -882,7 +883,7 @@ void UGameFeaturesSubsystem::LoadBuiltInGameFeaturePlugins(FBuiltInPluginAdditio
 	UAssetManager::Get().PopBulkScanning();
 }
 
-bool UGameFeaturesSubsystem::GetPluginURLForRegisteredPluginByName(const FString& PluginName, FString& OutPluginURL) const
+bool UGameFeaturesSubsystem::GetPluginURLByName(const FString& PluginName, FString& OutPluginURL) const
 {
 	if (const FString* PluginURL = GameFeaturePluginNameToPathMap.Find(PluginName))
 	{
@@ -895,7 +896,7 @@ bool UGameFeaturesSubsystem::GetPluginURLForRegisteredPluginByName(const FString
 
 bool UGameFeaturesSubsystem::GetPluginURLForBuiltInPluginByName(const FString& PluginName, FString& OutPluginURL) const
 {
-	return GetPluginURLForRegisteredPluginByName(PluginName, OutPluginURL);
+	return GetPluginURLByName(PluginName, OutPluginURL);
 }
 
 FString UGameFeaturesSubsystem::GetPluginFilenameFromPluginURL(const FString& PluginURL) const
diff --git a/Engine/Plugins/Experimental/GameFeatures/Source/GameFeatures/Public/GameFeaturesSubsystem.h b/Engine/Plugins/Experimental/GameFeatures/Source/GameFeatures/Public/GameFeaturesSubsystem.h
--- a/Engine/Plugins/Experimental/GameFeatures/Source/GameFeatures/Public/GameFeaturesSubsystem.h
+++ b/Engine/Plugins/Experimental/GameFeatures/Source/GameFeatures/Public/GameFeaturesSubsystem.h
@@ -264,8 +264,11 @@ public:
 	void TerminateGameFeaturePlugin(const FString& PluginURL);
 	void TerminateGameFeaturePlugin(const FString& PluginURL, const FGameFeaturePluginUninstallComplete& CompleteDelegate);
 
-	/** If the specified plugin is registered, return the URL used to identify it. Returns true if the plugin exists, false if it was not found */
-	bool GetPluginURLForRegisteredPluginByName(const FString& PluginName, FString& OutPluginURL) const;
+	/**
+	 * If the specified plugin is known by the game feature system, returns the URL used to identify it
+	 * @return true if the plugin exists, false if it was not found
+	 */
+	bool GetPluginURLByName(const FString& PluginName, FString& OutPluginURL) const;
 
 	// @TODO: Deprecate GetPluginURLForBuiltInPluginByName
 	/** If the specified plugin is a built-in plugin, return the URL used to identify it. Returns true if the plugin exists, false if it was not found */
@@ -313,12 +316,15 @@ public:
 private:
 	TSet<FString> GetActivePluginNames() const;
 
-	void OnGameFeatureTerminating(const FString& PluginURL);
+	void OnGameFeatureTerminating(const FString& PluginName, const FString& PluginURL);
 	friend struct FGameFeaturePluginState_Terminal;
 
 	void OnGameFeatureCheckingStatus(const FString& PluginURL);
 	friend struct FGameFeaturePluginState_UnknownStatus;
 
+	void OnGameFeatureStatusKnown(const FString& PluginName, const FString& PluginURL);
+	friend struct FGameFeaturePluginState_CheckingStatus;
+
 	void OnGameFeatureRegistering(const UGameFeatureData* GameFeatureData, const FString& PluginName, const FString& PluginURL);
 	friend struct FGameFeaturePluginState_Registering;
 
diff --git a/Engine/Plugins/Experimental/GameFeatures/Source/GameFeaturesEditor/Private/GameFeatureDataDetailsCustomization.cpp b/Engine/Plugins/Experimental/GameFeatures/Source/GameFeaturesEditor/Private/GameFeatureDataDetailsCustomization.cpp
--- a/Engine/Plugins/Experimental/GameFeatures/Source/GameFeaturesEditor/Private/GameFeatureDataDetailsCustomization.cpp
+++ b/Engine/Plugins/Experimental/GameFeatures/Source/GameFeaturesEditor/Private/GameFeatureDataDetailsCustomization.cpp
@@ -51,7 +51,7 @@ void FGameFeatureDataDetailsCustomization::CustomizeDetails(IDetailLayoutBuilder
 		GameFeature->GetOutermost()->GetName().ParseIntoArray(PathParts, TEXT("/"));
 
 		UGameFeaturesSubsystem& Subsystem = UGameFeaturesSubsystem::Get();
-		Subsystem.GetPluginURLForRegisteredPluginByName(PathParts[0], /*out*/ PluginURL);
+		Subsystem.GetPluginURLByName(PathParts[0], /*out*/ PluginURL);
 		PluginPtr = IPluginManager::Get().FindPlugin(PathParts[0]);
 
 		const float Padding = 8.0f;
diff --git a/Engine/Source/Runtime/GameplayTags/Private/NativeGameplayTags.cpp b/Engine/Source/Runtime/GameplayTags/Private/NativeGameplayTags.cpp
--- a/Engine/Source/Runtime/GameplayTags/Private/NativeGameplayTags.cpp
+++ b/Engine/Source/Runtime/GameplayTags/Private/NativeGameplayTags.cpp
@@ -26,7 +26,6 @@ static bool VerifyModuleCanContainGameplayTag(FName ModuleName, FName TagName, c
 			//check(ThisPlugin.IsValid());
 			//UGameplayTagsManager::Get().AddTagIniSearchPath(ThisPlugin->GetBaseDir() / TEXT("Config") / TEXT("Tags"));
 
-			//const FString PluginName = FPaths::GetBaseFilename(StateProperties.PluginInstalledFilename);
 			//const FString PluginFolder = FPaths::GetPath(StateProperties.PluginInstalledFilename);
 			//UGameplayTagsManager::Get().AddTagIniSearchPath(PluginFolder / TEXT("Config") / TEXT("Tags"));
 
