Add UGameFeaturesSubsystem::GetPluginURLProtocol

diff --git a/Engine/Plugins/Experimental/GameFeatures/Source/GameFeatures/Private/GameFeaturePluginStateMachine.h b/Engine/Plugins/Experimental/GameFeatures/Source/GameFeatures/Private/GameFeaturePluginStateMachine.h
--- a/Engine/Plugins/Experimental/GameFeatures/Source/GameFeatures/Private/GameFeaturePluginStateMachine.h
+++ b/Engine/Plugins/Experimental/GameFeatures/Source/GameFeatures/Private/GameFeaturePluginStateMachine.h
@@ -130,21 +130,6 @@ Transition states are expected to transition the machine to another state after
 
 // JMarcus TODO: move this stuff into UE::GameFeatures namespace
 
-#define GAME_FEATURE_PLUGIN_PROTOCOL_LIST(XPROTO)	\
-	XPROTO(File,			TEXT("file:"))			\
-	XPROTO(InstallBundle,	TEXT("installbundle:"))	\
-	XPROTO(Unknown,			TEXT(""))				
-
-#define GAME_FEATURE_PLUGIN_PROTOCOL_ENUM(inEnum, inString) inEnum,
-enum class EGameFeaturePluginProtocol : uint8
-{
-	GAME_FEATURE_PLUGIN_PROTOCOL_LIST(GAME_FEATURE_PLUGIN_PROTOCOL_ENUM)
-	Count,
-};
-#undef GAME_FEATURE_PLUGIN_PROTOCOL_ENUM
-
-ENUM_RANGE_BY_COUNT(EGameFeaturePluginProtocol, EGameFeaturePluginProtocol::Count);
-
 const TCHAR* GameFeaturePluginProtocolPrefix(EGameFeaturePluginProtocol Protocol);
 
 struct FInstallBundlePluginProtocolMetaData
diff --git a/Engine/Plugins/Experimental/GameFeatures/Source/GameFeatures/Private/GameFeaturesSubsystem.cpp b/Engine/Plugins/Experimental/GameFeatures/Source/GameFeatures/Private/GameFeaturesSubsystem.cpp
--- a/Engine/Plugins/Experimental/GameFeatures/Source/GameFeatures/Private/GameFeaturesSubsystem.cpp
+++ b/Engine/Plugins/Experimental/GameFeatures/Source/GameFeatures/Private/GameFeaturesSubsystem.cpp
@@ -376,6 +376,18 @@ FString UGameFeaturesSubsystem::GetPluginURL_InstallBundleProtocol(const FString
 	return GetPluginURL_InstallBundleProtocol(PluginName, MakeArrayView(&BundleName, 1));
 }
 
+EGameFeaturePluginProtocol UGameFeaturesSubsystem::GetPluginURLProtocol(FStringView PluginURL)
+{
+	for (EGameFeaturePluginProtocol Protocol : TEnumRange<EGameFeaturePluginProtocol>())
+	{
+		if (PluginURL.StartsWith(GameFeaturePluginProtocolPrefix(Protocol)))
+		{
+			return Protocol;
+		}
+	}
+	return EGameFeaturePluginProtocol::Unknown;
+}
+
 void UGameFeaturesSubsystem::OnGameFeatureTerminating(const FString& PluginName, const FString& PluginURL)
 {
 	for (UObject* Observer : Observers)
diff --git a/Engine/Plugins/Experimental/GameFeatures/Source/GameFeatures/Public/GameFeatureTypes.h b/Engine/Plugins/Experimental/GameFeatures/Source/GameFeatures/Public/GameFeatureTypes.h
--- a/Engine/Plugins/Experimental/GameFeatures/Source/GameFeatures/Public/GameFeatureTypes.h
+++ b/Engine/Plugins/Experimental/GameFeatures/Source/GameFeatures/Public/GameFeatureTypes.h
@@ -43,3 +43,18 @@ namespace GameFeaturePluginStatePrivate
 }
 typedef GameFeaturePluginStatePrivate::EGameFeaturePluginState EGameFeaturePluginState;
 #undef GAME_FEATURE_PLUGIN_STATE_ENUM
+
+#define GAME_FEATURE_PLUGIN_PROTOCOL_LIST(XPROTO)	\
+	XPROTO(File,			TEXT("file:"))			\
+	XPROTO(InstallBundle,	TEXT("installbundle:"))	\
+	XPROTO(Unknown,			TEXT(""))
+
+#define GAME_FEATURE_PLUGIN_PROTOCOL_ENUM(inEnum, inString) inEnum,
+enum class EGameFeaturePluginProtocol : uint8
+{
+	GAME_FEATURE_PLUGIN_PROTOCOL_LIST(GAME_FEATURE_PLUGIN_PROTOCOL_ENUM)
+	Count
+};
+#undef GAME_FEATURE_PLUGIN_PROTOCOL_ENUM
+
+ENUM_RANGE_BY_COUNT(EGameFeaturePluginProtocol, EGameFeaturePluginProtocol::Count);
diff --git a/Engine/Plugins/Experimental/GameFeatures/Source/GameFeatures/Public/GameFeatureTypesFwd.h b/Engine/Plugins/Experimental/GameFeatures/Source/GameFeatures/Public/GameFeatureTypesFwd.h
--- a/Engine/Plugins/Experimental/GameFeatures/Source/GameFeatures/Public/GameFeatureTypesFwd.h
+++ b/Engine/Plugins/Experimental/GameFeatures/Source/GameFeatures/Public/GameFeatureTypesFwd.h
@@ -2,6 +2,8 @@
 
 #pragma once
 
+enum class EGameFeaturePluginProtocol : uint8;
+
 namespace GameFeaturePluginStatePrivate
 {
 	enum EGameFeaturePluginState : uint8;
diff --git a/Engine/Plugins/Experimental/GameFeatures/Source/GameFeatures/Public/GameFeaturesSubsystem.h b/Engine/Plugins/Experimental/GameFeatures/Source/GameFeatures/Public/GameFeaturesSubsystem.h
--- a/Engine/Plugins/Experimental/GameFeatures/Source/GameFeatures/Public/GameFeaturesSubsystem.h
+++ b/Engine/Plugins/Experimental/GameFeatures/Source/GameFeatures/Public/GameFeaturesSubsystem.h
@@ -217,6 +217,8 @@ public:
 	static FString GetPluginURL_InstallBundleProtocol(const FString& PluginName, TArrayView<const FName> BundleNames);
 	static FString GetPluginURL_InstallBundleProtocol(const FString& PluginName, FName BundleName);
 
+	static EGameFeaturePluginProtocol GetPluginURLProtocol(FStringView PluginURL);
+
 public:
 	/** Returns all the active plugins GameFeatureDatas */
 	void GetGameFeatureDataForActivePlugins(TArray<const UGameFeatureData*>& OutActivePluginFeatureDatas);
