Deprecate UGameFeaturesSubsystem::GetPluginURLForBuiltInPluginByName

diff --git a/Engine/Plugins/Experimental/GameFeatures/Source/GameFeatures/Private/GameFeaturesSubsystem.cpp b/Engine/Plugins/Experimental/GameFeatures/Source/GameFeatures/Private/GameFeaturesSubsystem.cpp
--- a/Engine/Plugins/Experimental/GameFeatures/Source/GameFeatures/Private/GameFeaturesSubsystem.cpp
+++ b/Engine/Plugins/Experimental/GameFeatures/Source/GameFeatures/Private/GameFeaturesSubsystem.cpp
@@ -93,7 +93,7 @@ void UGameFeaturesSubsystem::Initialize(FSubsystemCollectionBase& Collection)
 			if (Args.Num() > 0)
 			{
 				FString PluginURL;
-				if (!UGameFeaturesSubsystem::Get().GetPluginURLForBuiltInPluginByName(Args[0], /*out*/ PluginURL))
+				if (!UGameFeaturesSubsystem::Get().GetPluginURLByName(Args[0], /*out*/ PluginURL))
 				{
 					PluginURL = Args[0];
 				}
@@ -114,7 +114,7 @@ void UGameFeaturesSubsystem::Initialize(FSubsystemCollectionBase& Collection)
 			if (Args.Num() > 0)
 			{
 				FString PluginURL;
-				if (!UGameFeaturesSubsystem::Get().GetPluginURLForBuiltInPluginByName(Args[0], /*out*/ PluginURL))
+				if (!UGameFeaturesSubsystem::Get().GetPluginURLByName(Args[0], /*out*/ PluginURL))
 				{
 					PluginURL = Args[0];
 				}
@@ -135,7 +135,7 @@ void UGameFeaturesSubsystem::Initialize(FSubsystemCollectionBase& Collection)
 			if (Args.Num() > 0)
 			{
 				FString PluginURL;
-				if (!UGameFeaturesSubsystem::Get().GetPluginURLForBuiltInPluginByName(Args[0], /*out*/ PluginURL))
+				if (!UGameFeaturesSubsystem::Get().GetPluginURLByName(Args[0], /*out*/ PluginURL))
 				{
 					PluginURL = Args[0];
 				}
@@ -156,7 +156,7 @@ void UGameFeaturesSubsystem::Initialize(FSubsystemCollectionBase& Collection)
 				if (Args.Num() > 0)
 				{
 					FString PluginURL;
-					if (!UGameFeaturesSubsystem::Get().GetPluginURLForBuiltInPluginByName(Args[0], /*out*/ PluginURL))
+					if (!UGameFeaturesSubsystem::Get().GetPluginURLByName(Args[0], /*out*/ PluginURL))
 					{
 						PluginURL = Args[0];
 					}
diff --git a/Engine/Plugins/Experimental/GameFeatures/Source/GameFeatures/Public/GameFeaturesSubsystem.h b/Engine/Plugins/Experimental/GameFeatures/Source/GameFeatures/Public/GameFeaturesSubsystem.h
--- a/Engine/Plugins/Experimental/GameFeatures/Source/GameFeatures/Public/GameFeaturesSubsystem.h
+++ b/Engine/Plugins/Experimental/GameFeatures/Source/GameFeatures/Public/GameFeaturesSubsystem.h
@@ -272,8 +272,8 @@ public:
 	 */
 	bool GetPluginURLByName(const FString& PluginName, FString& OutPluginURL) const;
 
-	// @TODO: Deprecate GetPluginURLForBuiltInPluginByName
 	/** If the specified plugin is a built-in plugin, return the URL used to identify it. Returns true if the plugin exists, false if it was not found */
+	UE_DEPRECATED(5.1, "Use GetPluginURLByName instead")
 	bool GetPluginURLForBuiltInPluginByName(const FString& PluginName, FString& OutPluginURL) const;
 
 	/** Get the plugin path from the plugin URL */
