Fix crash in GFP system when the termination callback would trigger a GC and destroy the GF state machine

diff --git a/Engine/Plugins/Experimental/GameFeatures/Source/GameFeatures/Private/GameFeaturePluginStateMachine.cpp b/Engine/Plugins/Experimental/GameFeatures/Source/GameFeatures/Private/GameFeaturePluginStateMachine.cpp
--- a/Engine/Plugins/Experimental/GameFeatures/Source/GameFeatures/Private/GameFeaturePluginStateMachine.cpp
+++ b/Engine/Plugins/Experimental/GameFeatures/Source/GameFeatures/Private/GameFeaturePluginStateMachine.cpp
@@ -137,8 +137,9 @@ void FGameFeaturePluginState::UpdateStateMachineDeferred(float Delay /*= 0.0f*/)
 
 	TickHandle = FTSTicker::GetCoreTicker().AddTicker(FTickerDelegate::CreateLambda([this](float dts) mutable
 	{
-		StateProperties.OnRequestUpdateStateMachine.ExecuteIfBound();
+		// @note Release FGameFeaturePluginState::TickHandle first in case the termination callback triggers a GC and destroys the state machine
 		TickHandle.Reset();
+		StateProperties.OnRequestUpdateStateMachine.ExecuteIfBound();
 		return false;
 	}), Delay);
 }
