Fix another crash in GFP system when the termination callback would trigger a GC and destroy the GF state machine (similar issue/fix as 19985514)

diff --git a/Engine/Plugins/Experimental/GameFeatures/Source/GameFeatures/Private/GameFeaturePluginStateMachine.cpp b/Engine/Plugins/Experimental/GameFeatures/Source/GameFeatures/Private/GameFeaturePluginStateMachine.cpp
--- a/Engine/Plugins/Experimental/GameFeatures/Source/GameFeatures/Private/GameFeaturePluginStateMachine.cpp
+++ b/Engine/Plugins/Experimental/GameFeatures/Source/GameFeatures/Private/GameFeaturePluginStateMachine.cpp
@@ -1707,8 +1707,11 @@ void UGameFeaturePluginStateMachine::UpdateStateMachine()
 	if (CurrentState == StateProperties.DestinationState)
 	{
 		check(IsValidTransitionState(CurrentState) == false);
-		StateProperties.OnFeatureStateTransitionComplete.ExecuteIfBound(this, TransitionResult);
-		StateProperties.OnFeatureStateTransitionComplete.Unbind();
+		// @note Clear out StateProperties.OnFeatureStateTransitionComplete prior to executing the 
+		// (termination) callback in case it triggers a GC that would destroy the state machine
+		FGameFeatureStateTransitionComplete TransitionCompleteCallback = MoveTemp(StateProperties.OnFeatureStateTransitionComplete);
+		StateProperties.OnFeatureStateTransitionComplete = FGameFeatureStateTransitionComplete();
+		TransitionCompleteCallback.ExecuteIfBound(this, TransitionResult);
 	}
 }
 
