Dump reference chains if an asset leak is detected on GFP unload.

diff --git a/Engine/Plugins/Experimental/GameFeatures/Source/GameFeatures/Private/GameFeaturePluginStateMachine.cpp b/Engine/Plugins/Experimental/GameFeatures/Source/GameFeatures/Private/GameFeaturePluginStateMachine.cpp
--- a/Engine/Plugins/Experimental/GameFeatures/Source/GameFeatures/Private/GameFeaturePluginStateMachine.cpp
+++ b/Engine/Plugins/Experimental/GameFeatures/Source/GameFeatures/Private/GameFeaturePluginStateMachine.cpp
@@ -22,6 +22,7 @@
 #include "GameFeaturesSubsystem.h"
 #include "GameFeaturesProjectPolicies.h"
 #include "Containers/Ticker.h"
+#include "UObject/ReferenceChainSearch.h"
 
 #if WITH_EDITOR
 #include "PluginUtils.h"
@@ -72,6 +73,19 @@ namespace UE::GameFeatures
 		PluginArFilter.PackagePaths.Add(FName(TEXT("/") + PluginName));
 		PluginArFilter.bRecursivePaths = true;
 
+		auto CheckForLoadedAsset = [](const FString& PluginName, const FAssetData &AssetData)
+		{
+			if (AssetData.IsAssetLoaded())
+			{
+				UE_LOG(LogGameFeatures, Error, TEXT("GFP %s failed to unload asset %s!"), *PluginName, *AssetData.GetFullName());
+
+				UObject* AssetObj = AssetData.GetAsset();
+				//EInternalObjectFlags InternalFlags = AssetObj->GetInternalFlags();
+
+				FReferenceChainSearch(AssetObj, EReferenceChainSearchMode::Shortest | EReferenceChainSearchMode::PrintResults);
+			}
+		};
+
 		if (bIgnoreGameFeatureData)
 		{
 			FARFilter RawGameFeatureDataFilter;
@@ -81,29 +95,23 @@ namespace UE::GameFeatures
 			FARCompiledFilter GameFeatureDataFilter;
 			UAssetManager::Get().GetAssetRegistry().CompileFilter(RawGameFeatureDataFilter, GameFeatureDataFilter);
 
-			UAssetManager::Get().GetAssetRegistry().EnumerateAssets(PluginArFilter, [&PluginName, &GameFeatureDataFilter](const FAssetData& AssetData)
+			UAssetManager::Get().GetAssetRegistry().EnumerateAssets(PluginArFilter, [&PluginName, &GameFeatureDataFilter, CheckForLoadedAsset](const FAssetData& AssetData)
 				{
 					if (UAssetManager::Get().GetAssetRegistry().IsAssetIncludedByFilter(AssetData, GameFeatureDataFilter))
 					{
 						return true;
 					}
 
-					if (AssetData.IsAssetLoaded())
-					{
-						UE_LOG(LogGameFeatures, Error, TEXT("GFP %s failed to unload asset %s!"), *PluginName, *AssetData.GetFullName());
-					}
+					CheckForLoadedAsset(PluginName, AssetData);
 
 					return true;
 				});
 		}
 		else
 		{
-			UAssetManager::Get().GetAssetRegistry().EnumerateAssets(PluginArFilter, [&PluginName](const FAssetData& AssetData)
+			UAssetManager::Get().GetAssetRegistry().EnumerateAssets(PluginArFilter, [&PluginName, CheckForLoadedAsset](const FAssetData& AssetData)
 				{
-					if (AssetData.IsAssetLoaded())
-					{
-						UE_LOG(LogGameFeatures, Error, TEXT("GFP %s failed to unload asset %s!"), *PluginName, *AssetData.GetFullName());
-					}
+					CheckForLoadedAsset(PluginName, AssetData);
 
 					return true;
 				});
