Remove GameFeaturePluginURL field from uplugin dependency and use UGameFeaturesSubsystem::GetPluginURLByName instead

diff --git a/Engine/Plugins/Experimental/GameFeatures/Source/GameFeatures/Private/GameFeaturesSubsystem.cpp b/Engine/Plugins/Experimental/GameFeatures/Source/GameFeatures/Private/GameFeaturesSubsystem.cpp
--- a/Engine/Plugins/Experimental/GameFeatures/Source/GameFeatures/Private/GameFeaturesSubsystem.cpp
+++ b/Engine/Plugins/Experimental/GameFeatures/Source/GameFeatures/Private/GameFeaturesSubsystem.cpp
@@ -1015,7 +1015,6 @@ bool UGameFeaturesSubsystem::GetGameFeaturePluginDetails(const FString& PluginDe
 	{
 		FString NameField = TEXT("Name");
 		FString EnabledField = TEXT("Enabled");
-		FString GameFeaturePluginURLField = TEXT("GameFeaturePluginURL");
 		for (const TSharedPtr<FJsonValue>& PluginElement : *PluginsArray)
 		{
 			if (PluginElement.IsValid())
@@ -1030,35 +1029,33 @@ bool UGameFeaturesSubsystem::GetGameFeaturePluginDetails(const FString& PluginDe
 
 					if (bElementEnabled)
 					{
-						FString GameFeaturePluginURL;
-						ElementObject->TryGetStringField(GameFeaturePluginURLField, GameFeaturePluginURL);
-						if (!GameFeaturePluginURL.IsEmpty())
+						FString DependencyName;
+						ElementObject->TryGetStringField(NameField, DependencyName);
+						if (!DependencyName.IsEmpty())
 						{
-							OutPluginDetails.PluginDependencies.Add(GameFeaturePluginURL);
-						}
-						else
-						{
-							FString ElementName;
-							ElementObject->TryGetStringField(NameField, ElementName);
-							if (!ElementName.IsEmpty())
+							TSharedPtr<IPlugin> DependencyPlugin = IPluginManager::Get().FindPlugin(DependencyName);
+							if (DependencyPlugin.IsValid())
 							{
-								TSharedPtr<IPlugin> DependencyPlugin = IPluginManager::Get().FindPlugin(ElementName);
-								if (DependencyPlugin.IsValid())
+								FString DependencyURL;
+								if (!GetPluginURLByName(DependencyPlugin->GetName(), DependencyURL))
 								{
-									const FString& PluginDependencyDescriptorFilename = DependencyPlugin->GetDescriptorFileName();
-
-									if (!PluginDependencyDescriptorFilename.IsEmpty() &&
-										GetDefault<UGameFeaturesSubsystemSettings>()->IsValidGameFeaturePlugin(FPaths::ConvertRelativePathToFull(PluginDependencyDescriptorFilename)) &&
-										FPaths::FileExists(PluginDependencyDescriptorFilename))
+									if (!DependencyPlugin->GetDescriptorFileName().IsEmpty() &&
+										GetDefault<UGameFeaturesSubsystemSettings>()->IsValidGameFeaturePlugin(FPaths::ConvertRelativePathToFull(DependencyPlugin->GetDescriptorFileName())) &&
+										FPaths::FileExists(DependencyPlugin->GetDescriptorFileName()))
 									{
-										OutPluginDetails.PluginDependencies.Add(GetPluginURL_FileProtocol(DependencyPlugin->GetDescriptorFileName()));
+										DependencyURL = GetPluginURL_FileProtocol(DependencyPlugin->GetDescriptorFileName());
 									}
 								}
-								else
+
+								if (!DependencyURL.IsEmpty())
 								{
-									UE_LOG(LogGameFeatures, Display, TEXT("Game feature plugin '%s' has unknown dependency '%s'."), *PluginDescriptorFilename, *ElementName);
+									OutPluginDetails.PluginDependencies.Add(DependencyURL);
 								}
 							}
+							else
+							{
+								UE_LOG(LogGameFeatures, Display, TEXT("Game feature plugin '%s' has unknown dependency '%s'."), *PluginDescriptorFilename, *DependencyName);
+							}
 						}
 					}
 				}
