Change GameFeaturePlugin.VerifyUnload to ECVF_Default

diff --git a/Engine/Plugins/Experimental/GameFeatures/Source/GameFeatures/Private/GameFeaturePluginStateMachine.cpp b/Engine/Plugins/Experimental/GameFeatures/Source/GameFeatures/Private/GameFeaturePluginStateMachine.cpp
--- a/Engine/Plugins/Experimental/GameFeatures/Source/GameFeatures/Private/GameFeaturePluginStateMachine.cpp
+++ b/Engine/Plugins/Experimental/GameFeatures/Source/GameFeatures/Private/GameFeaturePluginStateMachine.cpp
@@ -40,7 +40,7 @@ namespace UE::GameFeatures
 	static TAutoConsoleVariable<bool> CVarVerifyPluginUnload(TEXT("GameFeaturePlugin.VerifyUnload"), 
 		false,
 		TEXT("Verify plugin assets are no longer in memory when unloading."),
-		ECVF_Cheat);
+		ECVF_Default);
 
 	FString ToString(const UE::GameFeatures::FResult& Result)
 	{
