Refactor to prep for GFP canceling and to make GFP state machine transition requests more robust.

Solved Problems:
It was impossible to know if it was ok to change destination states if the state machine was currently running.
If the state machine was running, setting a new destination would stomp any existing callbacks without calling them.
If was impossible to know if the current tranistion needed be canceld because there was no way to tell if the new destination state was complatible with the intent of the current state transition.

Changes:
- Instead of passing a single destination state pass a range acceptable destination states to the machine.  This communicates the intent of the caller.
- If the machine is running, then the intersection of any new range and the existing range will satisfy both requests.
- If there is no intersection, the request can't be satisified and is refused (I'll change this to an actual implementation of canceling in a future change list)
- The state machine tracks multiple completion callbacks, callbacks are no longer stomped and are always called back.
- Stopped deferring a lot of callbacks to next tick since there were still a lot of cases where callbacks could happen immediatlely and it was giving a false sense of the way GFPs worked.  Content Beacon V2 should be fine with syncronous callbacks and CBV1 is no longer used. I couldn't find other cases that requried deffered callbacks.
- Clear bInUpdateStateMachine before the state machine calls back so that callbacks can force new state machine updates. This was a common cause of adding deffered callbacks in some cases.
- Since machines can now track multiple callbacks, FGameFeaturePluginState_WaitingForDependencies no longer needs special case state transition callbacks
- Moved a few things into the UE::GameFeatures namespace.

diff --git a/Engine/Plugins/Experimental/GameFeatures/Source/GameFeatures/Private/GameFeaturePluginStateMachine.cpp b/Engine/Plugins/Experimental/GameFeatures/Source/GameFeatures/Private/GameFeaturePluginStateMachine.cpp
--- a/Engine/Plugins/Experimental/GameFeatures/Source/GameFeatures/Private/GameFeaturePluginStateMachine.cpp
+++ b/Engine/Plugins/Experimental/GameFeatures/Source/GameFeatures/Private/GameFeaturePluginStateMachine.cpp
@@ -119,26 +119,51 @@ namespace UE::GameFeatures
 
 #endif // UE_BUILD_SHIPPING
 	}
-}
 
 #define GAME_FEATURE_PLUGIN_PROTOCOL_PREFIX(inEnum, inString) case EGameFeaturePluginProtocol::inEnum: return inString;
-const TCHAR* GameFeaturePluginProtocolPrefix(EGameFeaturePluginProtocol Protocol)
-{
-	switch (Protocol)
+	const TCHAR* GameFeaturePluginProtocolPrefix(EGameFeaturePluginProtocol Protocol)
 	{
-		GAME_FEATURE_PLUGIN_PROTOCOL_LIST(GAME_FEATURE_PLUGIN_PROTOCOL_PREFIX)
-	}
+		switch (Protocol)
+		{
+			GAME_FEATURE_PLUGIN_PROTOCOL_LIST(GAME_FEATURE_PLUGIN_PROTOCOL_PREFIX)
+		}
 
-	check(false);
-	return nullptr;
-}
+		check(false);
+		return nullptr;
+	}
 #undef GAME_FEATURE_PLUGIN_PROTOCOL_PREFIX
+}
 
 FGameFeaturePluginState::~FGameFeaturePluginState()
 {
 	CleanupDeferredUpdateCallbacks();
 }
 
+FDestinationGameFeaturePluginState* FGameFeaturePluginState::AsDestinationState()
+{
+	FDestinationGameFeaturePluginState* Ret = nullptr;
+
+	EGameFeaturePluginStateType Type = GetStateType();
+	if (Type == EGameFeaturePluginStateType::Destination || Type == EGameFeaturePluginStateType::Error)
+	{
+		Ret = static_cast<FDestinationGameFeaturePluginState*>(this);
+	}
+
+	return Ret;
+}
+
+FErrorGameFeaturePluginState* FGameFeaturePluginState::AsErrorState()
+{
+	FErrorGameFeaturePluginState* Ret = nullptr;
+
+	if (GetStateType() == EGameFeaturePluginStateType::Error)
+	{
+		Ret = static_cast<FErrorGameFeaturePluginState*>(this);
+	}
+
+	return Ret;
+}
+
 void FGameFeaturePluginState::UpdateStateMachineDeferred(float Delay /*= 0.0f*/) const
 {
 	CleanupDeferredUpdateCallbacks();
@@ -187,26 +212,6 @@ void FGameFeaturePluginState::CleanupDeferredUpdateCallbacks() const
 =========================================================
 */
 
-struct FDestinationGameFeaturePluginState : public FGameFeaturePluginState
-{
-	FDestinationGameFeaturePluginState(FGameFeaturePluginStateMachineProperties& InStateProperties) : FGameFeaturePluginState(InStateProperties) {}
-
-	virtual EGameFeaturePluginStateType GetStateType() const override
-	{
-		return EGameFeaturePluginStateType::Destination;
-	}
-};
-
-struct FErrorGameFeaturePluginState : public FGameFeaturePluginState
-{
-	FErrorGameFeaturePluginState(FGameFeaturePluginStateMachineProperties& InStateProperties) : FGameFeaturePluginState(InStateProperties) {}
-
-	virtual EGameFeaturePluginStateType GetStateType() const override
-	{
-		return EGameFeaturePluginStateType::Error;
-	}
-};
-
 struct FGameFeaturePluginState_Uninitialized : public FGameFeaturePluginState
 {
 	FGameFeaturePluginState_Uninitialized(FGameFeaturePluginStateMachineProperties& InStateProperties) : FGameFeaturePluginState(InStateProperties) {}
@@ -238,11 +243,11 @@ struct FGameFeaturePluginState_UnknownStatus : public FDestinationGameFeaturePlu
 
 	virtual void UpdateState(FGameFeaturePluginStateStatus& StateStatus) override
 	{
-		if (StateProperties.DestinationState < EGameFeaturePluginState::UnknownStatus)
+		if (StateProperties.Destination < EGameFeaturePluginState::UnknownStatus)
 		{
 			StateStatus.SetTransition(EGameFeaturePluginState::Terminal);
 		}
-		else if (StateProperties.DestinationState > EGameFeaturePluginState::UnknownStatus)
+		else if (StateProperties.Destination > EGameFeaturePluginState::UnknownStatus)
 		{
 			StateStatus.SetTransition(EGameFeaturePluginState::CheckingStatus);
 
@@ -302,7 +307,7 @@ struct FGameFeaturePluginState_CheckingStatus : public FGameFeaturePluginState
 				return;
 			}
 
-			const TArray<FName>& InstallBundles = StateProperties.ProtocolMetadata.GetSubtype<FInstallBundlePluginProtocolMetaData>().InstallBundles;
+			const TArray<FName>& InstallBundles = StateProperties.ProtocolMetadata.GetSubtype<UE::GameFeatures::FInstallBundlePluginProtocolMetaData>().InstallBundles;
 
 			TValueOrError<FInstallBundleCombinedInstallState, EInstallBundleResult> MaybeInstallState = BundleManager->GetInstallStateSynchronous(InstallBundles, false);
 			if (MaybeInstallState.HasError())
@@ -337,11 +342,11 @@ struct FGameFeaturePluginState_ErrorCheckingStatus : public FErrorGameFeaturePlu
 
 	virtual void UpdateState(FGameFeaturePluginStateStatus& StateStatus) override
 	{
-		if (StateProperties.DestinationState < EGameFeaturePluginState::ErrorCheckingStatus)
+		if (StateProperties.Destination < EGameFeaturePluginState::ErrorCheckingStatus)
 		{
 			StateStatus.SetTransition(EGameFeaturePluginState::Terminal);
 		}
-		else
+		else if (StateProperties.Destination > EGameFeaturePluginState::ErrorCheckingStatus)
 		{
 			StateStatus.SetTransition(EGameFeaturePluginState::CheckingStatus);
 		}
@@ -354,11 +359,11 @@ struct FGameFeaturePluginState_ErrorUnavailable : public FErrorGameFeaturePlugin
 
 	virtual void UpdateState(FGameFeaturePluginStateStatus& StateStatus) override
 	{
-		if (StateProperties.DestinationState < EGameFeaturePluginState::ErrorUnavailable)
+		if (StateProperties.Destination < EGameFeaturePluginState::ErrorUnavailable)
 		{
 			StateStatus.SetTransition(EGameFeaturePluginState::Terminal);
 		}
-		else
+		else if (StateProperties.Destination > EGameFeaturePluginState::ErrorUnavailable)
 		{
 			StateStatus.SetTransition(EGameFeaturePluginState::CheckingStatus);
 		}
@@ -371,11 +376,11 @@ struct FGameFeaturePluginState_StatusKnown : public FDestinationGameFeaturePlugi
 
 	virtual void UpdateState(FGameFeaturePluginStateStatus& StateStatus) override
 	{
-		if (StateProperties.DestinationState < EGameFeaturePluginState::StatusKnown)
+		if (StateProperties.Destination < EGameFeaturePluginState::StatusKnown)
 		{
 			StateStatus.SetTransition(EGameFeaturePluginState::Terminal);
 		}
-		else if (StateProperties.DestinationState > EGameFeaturePluginState::StatusKnown)
+		else if (StateProperties.Destination > EGameFeaturePluginState::StatusKnown)
 		{
 			if (StateProperties.GetPluginProtocol() != EGameFeaturePluginProtocol::File)
 			{
@@ -395,11 +400,11 @@ struct FGameFeaturePluginState_ErrorInstalling : public FErrorGameFeaturePluginS
 
 	virtual void UpdateState(FGameFeaturePluginStateStatus& StateStatus) override
 	{
-		if (StateProperties.DestinationState < EGameFeaturePluginState::ErrorInstalling)
+		if (StateProperties.Destination < EGameFeaturePluginState::ErrorInstalling)
 		{
 			StateStatus.SetTransition(EGameFeaturePluginState::Uninstalling);
 		}
-		else
+		else if (StateProperties.Destination > EGameFeaturePluginState::ErrorInstalling)
 		{
 			StateStatus.SetTransition(EGameFeaturePluginState::Downloading);
 		}
@@ -457,7 +462,7 @@ struct FGameFeaturePluginState_Uninstalling : public FGameFeaturePluginState
 
 		TSharedPtr<IInstallBundleManager> BundleManager = IInstallBundleManager::GetPlatformInstallBundleManager();
 
-		const TArray<FName>& InstallBundles = StateProperties.ProtocolMetadata.GetSubtype<FInstallBundlePluginProtocolMetaData>().InstallBundles;
+		const TArray<FName>& InstallBundles = StateProperties.ProtocolMetadata.GetSubtype<UE::GameFeatures::FInstallBundlePluginProtocolMetaData>().InstallBundles;
 
 		EInstallBundleReleaseRequestFlags ReleaseFlags = EInstallBundleReleaseRequestFlags::RemoveFilesIfPossible;
 		TValueOrError<FInstallBundleRequestInfo, EInstallBundleResult> MaybeRequestInfo = BundleManager->RequestReleaseContent(InstallBundles, ReleaseFlags);
@@ -561,7 +566,7 @@ struct FGameFeaturePluginState_Downloading : public FGameFeaturePluginState
 		GotContentStateHandle.Reset();
 
 		TSharedPtr<IInstallBundleManager> BundleManager = IInstallBundleManager::GetPlatformInstallBundleManager();
-		const TArray<FName>& InstallBundles = StateProperties.ProtocolMetadata.GetSubtype<FInstallBundlePluginProtocolMetaData>().InstallBundles;
+		const TArray<FName>& InstallBundles = StateProperties.ProtocolMetadata.GetSubtype<UE::GameFeatures::FInstallBundlePluginProtocolMetaData>().InstallBundles;
 
 		EInstallBundleRequestFlags InstallFlags = EInstallBundleRequestFlags::None;
 		InstallFlags |= EInstallBundleRequestFlags::SkipMount;
@@ -662,7 +667,7 @@ struct FGameFeaturePluginState_Downloading : public FGameFeaturePluginState
 		check(StateProperties.GetPluginProtocol() == EGameFeaturePluginProtocol::InstallBundle);
 
 		TSharedPtr<IInstallBundleManager> BundleManager = IInstallBundleManager::GetPlatformInstallBundleManager();
-		const TArray<FName>& InstallBundles = StateProperties.ProtocolMetadata.GetSubtype<FInstallBundlePluginProtocolMetaData>().InstallBundles;
+		const TArray<FName>& InstallBundles = StateProperties.ProtocolMetadata.GetSubtype<UE::GameFeatures::FInstallBundlePluginProtocolMetaData>().InstallBundles;
 
 		GotContentStateHandle = BundleManager->GetContentState(InstallBundles, EInstallBundleGetContentStateFlags::None, true, FInstallBundleGetContentStateDelegate::CreateRaw(this, &FGameFeaturePluginState_Downloading::OnGotContentState));
 	}
@@ -695,11 +700,11 @@ struct FGameFeaturePluginState_Installed : public FDestinationGameFeaturePluginS
 
 	virtual void UpdateState(FGameFeaturePluginStateStatus& StateStatus) override
 	{
-		if (StateProperties.DestinationState > EGameFeaturePluginState::Installed)
+		if (StateProperties.Destination > EGameFeaturePluginState::Installed)
 		{
 			StateStatus.SetTransition(EGameFeaturePluginState::Mounting);
 		}
-		else if (StateProperties.DestinationState < EGameFeaturePluginState::Installed)
+		else if (StateProperties.Destination < EGameFeaturePluginState::Installed)
 		{
 			StateStatus.SetTransition(EGameFeaturePluginState::Uninstalling);
 		}
@@ -712,11 +717,11 @@ struct FGameFeaturePluginState_ErrorMounting : public FErrorGameFeaturePluginSta
 
 	virtual void UpdateState(FGameFeaturePluginStateStatus& StateStatus) override
 	{
-		if (StateProperties.DestinationState < EGameFeaturePluginState::ErrorMounting)
+		if (StateProperties.Destination < EGameFeaturePluginState::ErrorMounting)
 		{
 			StateStatus.SetTransition(EGameFeaturePluginState::Unmounting);
 		}
-		else
+		else if (StateProperties.Destination > EGameFeaturePluginState::ErrorMounting)
 		{
 			StateStatus.SetTransition(EGameFeaturePluginState::Mounting);
 		}
@@ -729,12 +734,12 @@ struct FGameFeaturePluginState_ErrorWaitingForDependencies : public FErrorGameFe
 
 	virtual void UpdateState(FGameFeaturePluginStateStatus& StateStatus) override
 	{
-		if (StateProperties.DestinationState < EGameFeaturePluginState::ErrorWaitingForDependencies)
+		if (StateProperties.Destination < EGameFeaturePluginState::ErrorWaitingForDependencies)
 		{
 			// There is no cleaup state equivalent to EGameFeaturePluginState::WaitingForDependencies so just go back to unmounting
 			StateStatus.SetTransition(EGameFeaturePluginState::Unmounting);
 		}
-		else
+		else if (StateProperties.Destination > EGameFeaturePluginState::ErrorWaitingForDependencies)
 		{
 			StateStatus.SetTransition(EGameFeaturePluginState::WaitingForDependencies);
 		}
@@ -747,11 +752,11 @@ struct FGameFeaturePluginState_ErrorRegistering : public FErrorGameFeaturePlugin
 
 	virtual void UpdateState(FGameFeaturePluginStateStatus& StateStatus) override
 	{
-		if (StateProperties.DestinationState < EGameFeaturePluginState::ErrorRegistering)
+		if (StateProperties.Destination < EGameFeaturePluginState::ErrorRegistering)
 		{
 			StateStatus.SetTransition(EGameFeaturePluginState::Unregistering);
 		}
-		else
+		else if (StateProperties.Destination > EGameFeaturePluginState::ErrorRegistering)
 		{
 			StateStatus.SetTransition(EGameFeaturePluginState::Registering);
 		}
@@ -821,7 +826,7 @@ struct FGameFeaturePluginState_Unmounting : public FGameFeaturePluginState
 
 		TSharedPtr<IInstallBundleManager> BundleManager = IInstallBundleManager::GetPlatformInstallBundleManager();
 
-		const TArray<FName>& InstallBundles = StateProperties.ProtocolMetadata.GetSubtype<FInstallBundlePluginProtocolMetaData>().InstallBundles;
+		const TArray<FName>& InstallBundles = StateProperties.ProtocolMetadata.GetSubtype<UE::GameFeatures::FInstallBundlePluginProtocolMetaData>().InstallBundles;
 
 		EInstallBundleReleaseRequestFlags ReleaseFlags = EInstallBundleReleaseRequestFlags::None;
 		TValueOrError<FInstallBundleRequestInfo, EInstallBundleResult> MaybeRequestInfo = BundleManager->RequestReleaseContent(InstallBundles, ReleaseFlags);
@@ -948,7 +953,7 @@ struct FGameFeaturePluginState_Mounting : public FGameFeaturePluginState
 		
 		TSharedPtr<IInstallBundleManager> BundleManager = IInstallBundleManager::GetPlatformInstallBundleManager();
 
-		const TArray<FName>& InstallBundles = StateProperties.ProtocolMetadata.GetSubtype<FInstallBundlePluginProtocolMetaData>().InstallBundles;
+		const TArray<FName>& InstallBundles = StateProperties.ProtocolMetadata.GetSubtype<UE::GameFeatures::FInstallBundlePluginProtocolMetaData>().InstallBundles;
 		
 		// JMarcus TODO: Async Mounting?
 		EInstallBundleRequestFlags InstallFlags = EInstallBundleRequestFlags::None;
@@ -1095,52 +1100,46 @@ struct FGameFeaturePluginState_WaitingForDependencies : public FGameFeaturePlugi
 		if (!bRequestedDependencies)
 		{
 			TArray<UGameFeaturePluginStateMachine*> Dependencies;
-			check(StateProperties.OnRequestStateMachineDependencies.IsBound());
-			if (StateProperties.OnRequestStateMachineDependencies.Execute(StateProperties.PluginInstalledFilename, Dependencies))
-			{
-				bRequestedDependencies = true;
-				for (UGameFeaturePluginStateMachine* Dependency : Dependencies)
-				{
-					check(Dependency);
-					if (Dependency->GetCurrentState() < EGameFeaturePluginState::Registered)
-					{
-						RemainingDependencies.Add(Dependency);
-						Dependency->OnStateChanged().AddRaw(this, &FGameFeaturePluginState_WaitingForDependencies::OnDependencyStateChanged);
-
-						// If we are not alreadying loading this dependency, do so now
-						if (Dependency->GetDestinationState() < EGameFeaturePluginState::Registered)
-						{
-							Dependency->SetDestinationState(EGameFeaturePluginState::Registered, FGameFeatureStateTransitionComplete());
-						}
-					}
-				}
-			}
-			else
+			if (!UGameFeaturesSubsystem::Get().FindOrCreatePluginDependencyStateMachines(StateProperties.PluginInstalledFilename, Dependencies))
 			{
 				// Failed to query dependencies
 				StateStatus.SetTransitionError(EGameFeaturePluginState::ErrorWaitingForDependencies, UE::GameFeatures::StateMachineErrorNamespace + TEXT("Failed_Dependency_Query"));
 				return;
 			}
+			
+			bRequestedDependencies = true;
+
+			RemainingDependencies.Reserve(Dependencies.Num());
+			for (UGameFeaturePluginStateMachine* Dependency : Dependencies)
+			{	
+				RemainingDependencies.Emplace(Dependency, MakeValue());
+
+				bool bSetDestination = Dependency->SetDestination(
+					FGameFeaturePluginStateRange(EGameFeaturePluginState::Registered, EGameFeaturePluginState::Active), 
+					FGameFeatureStateTransitionComplete::CreateRaw(this, &FGameFeaturePluginState_WaitingForDependencies::OnDependencyTransitionComplete));
+
+				if (!bSetDestination)
+				{
+					StateStatus.SetTransitionError(EGameFeaturePluginState::ErrorWaitingForDependencies, UE::GameFeatures::StateMachineErrorNamespace + TEXT("Failed_Dependency_Register"));
+					return;
+				}
+			}
 		}
 
-		for (int32 DepIdx = RemainingDependencies.Num() - 1; DepIdx >= 0; --DepIdx)
+		for (FDepResultPair& Pair : RemainingDependencies)
 		{
-			UGameFeaturePluginStateMachine* RemainingDependency = RemainingDependencies[DepIdx].Get();
+			UGameFeaturePluginStateMachine* RemainingDependency = Pair.Key.Get();
 			if (!RemainingDependency)
 			{
 				// One of the dependency state machines was destroyed before finishing
 				StateStatus.SetTransitionError(EGameFeaturePluginState::ErrorWaitingForDependencies, UE::GameFeatures::StateMachineErrorNamespace + TEXT("Dependency_Destroyed_Before_Finish"));
 				return;
 			}
-			else if (RemainingDependency->GetCurrentState() >= EGameFeaturePluginState::Registered)
-			{
-				RemainingDependency->OnStateChanged().RemoveAll(this);
-				RemainingDependencies.RemoveAt(DepIdx, 1, false);
-			}
-			else if (RemainingDependency->GetCurrentState() == RemainingDependency->GetDestinationState())
+
+			if (Pair.Value.HasError())
 			{
-				// The dependency is no longer transitioning and is not Registered or later, so it failed to register, thus we cannot proceed
 				StateStatus.SetTransitionError(EGameFeaturePluginState::ErrorWaitingForDependencies, UE::GameFeatures::StateMachineErrorNamespace + TEXT("Failed_Dependency_Register"));
+				return;
 			}
 		}
 
@@ -1150,29 +1149,48 @@ struct FGameFeaturePluginState_WaitingForDependencies : public FGameFeaturePlugi
 		}
 	}
 
-	void OnDependencyStateChanged(UGameFeaturePluginStateMachine* Dependency)
+	void OnDependencyTransitionComplete(UGameFeaturePluginStateMachine* Dependency, const UE::GameFeatures::FResult& Result)
 	{
-		if (RemainingDependencies.Contains(Dependency))
+		const int32 Index = RemainingDependencies.IndexOfByPredicate([Dependency](const FDepResultPair& Pair)
+		{
+			return Pair.Key == Dependency;
+		});
+
+		if (Index != INDEX_NONE)
 		{
+			FDepResultPair& FoundDep = RemainingDependencies[Index];
+			
+			if (Result.HasError())
+			{
+				FoundDep.Value = Result;
+			}
+			else
+			{
+				RemainingDependencies.RemoveAtSwap(Index, 1, false);
+			}
+
 			UpdateStateMachineImmediate();
 		}
 	}
 
 	void ClearDependencies()
 	{
-		for (const TWeakObjectPtr<UGameFeaturePluginStateMachine>& WeakDependency : RemainingDependencies)
+		for (FDepResultPair& Pair : RemainingDependencies)
 		{
-			if (UGameFeaturePluginStateMachine* Dependency = WeakDependency.Get())
+			UGameFeaturePluginStateMachine* RemainingDependency = Pair.Key.Get();
+			if (RemainingDependency)
 			{
-				Dependency->OnStateChanged().RemoveAll(this);
+				RemainingDependency->RemovePendingTransitionCallback(this);
 			}
 		}
+
 		RemainingDependencies.Empty();
 		bRequestedDependencies = false;
 	}
 
-	TArray<TWeakObjectPtr<UGameFeaturePluginStateMachine>> RemainingDependencies;
-	bool bRequestedDependencies;
+	using FDepResultPair = TPair<TWeakObjectPtr<UGameFeaturePluginStateMachine>, UE::GameFeatures::FResult>;
+	TArray<FDepResultPair> RemainingDependencies;
+	bool bRequestedDependencies = false;
 };
 
 struct FGameFeaturePluginState_Unregistering : public FGameFeaturePluginState
@@ -1279,11 +1297,11 @@ struct FGameFeaturePluginState_Registered : public FDestinationGameFeaturePlugin
 
 	virtual void UpdateState(FGameFeaturePluginStateStatus& StateStatus) override
 	{
-		if (StateProperties.DestinationState > EGameFeaturePluginState::Registered)
+		if (StateProperties.Destination > EGameFeaturePluginState::Registered)
 		{
 			StateStatus.SetTransition(EGameFeaturePluginState::Loading);
 		}
-		else if (StateProperties.DestinationState < EGameFeaturePluginState::Registered)
+		else if (StateProperties.Destination < EGameFeaturePluginState::Registered)
 		{
 			StateStatus.SetTransition( EGameFeaturePluginState::Unregistering);
 		}
@@ -1315,7 +1333,7 @@ struct FGameFeaturePluginState_Unloading : public FGameFeaturePluginState
 
 		UnloadGameFeatureBundles(StateProperties.GameFeatureData);
 
-		if (StateProperties.DestinationState == EGameFeaturePluginState::Registered)
+		if (StateProperties.Destination.MaxState == EGameFeaturePluginState::Registered)
 		{
 			// If we aren't going farther than Registered, GC now
 			// otherwise we will defer until closer to our destination state
@@ -1405,11 +1423,11 @@ struct FGameFeaturePluginState_Loaded : public FDestinationGameFeaturePluginStat
 
 	virtual void UpdateState(FGameFeaturePluginStateStatus& StateStatus) override
 	{
-		if (StateProperties.DestinationState > EGameFeaturePluginState::Loaded)
+		if (StateProperties.Destination > EGameFeaturePluginState::Loaded)
 		{
 			StateStatus.SetTransition(EGameFeaturePluginState::Activating);
 		}
-		else if (StateProperties.DestinationState < EGameFeaturePluginState::Loaded)
+		else if (StateProperties.Destination < EGameFeaturePluginState::Loaded)
 		{
 			StateStatus.SetTransition(EGameFeaturePluginState::Unloading);
 		}
@@ -1468,7 +1486,7 @@ struct FGameFeaturePluginState_Deactivating : public FGameFeaturePluginState
 
 		if (NumExpectedPausers == NumObservedPausers)
 		{
-			if(!bRequestedGC && StateProperties.DestinationState == EGameFeaturePluginState::Loaded)
+			if(!bRequestedGC && StateProperties.Destination.MaxState == EGameFeaturePluginState::Loaded)
 			{
 				// If we aren't going farther than Loaded, GC now
 				// otherwise we will defer until closer to our destination state
@@ -1512,7 +1530,7 @@ struct FGameFeaturePluginState_Active : public FDestinationGameFeaturePluginStat
 
 	virtual void UpdateState(FGameFeaturePluginStateStatus& StateStatus) override
 	{
-		if (StateProperties.DestinationState < EGameFeaturePluginState::Active)
+		if (StateProperties.Destination < EGameFeaturePluginState::Active)
 		{
 			StateStatus.SetTransition(EGameFeaturePluginState::Deactivating);
 		}
@@ -1533,14 +1551,13 @@ UGameFeaturePluginStateMachine::UGameFeaturePluginStateMachine(const FObjectInit
 
 }
 
-void UGameFeaturePluginStateMachine::InitStateMachine(const FString& InPluginURL, const FGameFeaturePluginRequestStateMachineDependencies& OnRequestStateMachineDependencies)
+void UGameFeaturePluginStateMachine::InitStateMachine(const FString& InPluginURL)
 {
 	check(GetCurrentState() == EGameFeaturePluginState::Uninitialized);
 	CurrentStateInfo.State = EGameFeaturePluginState::UnknownStatus;
 	StateProperties = FGameFeaturePluginStateMachineProperties(
 		InPluginURL,
-		CurrentStateInfo.State,
-		OnRequestStateMachineDependencies,
+		FGameFeaturePluginStateRange(CurrentStateInfo.State),
 		FGameFeaturePluginRequestUpdateStateMachine::CreateUObject(this, &ThisClass::UpdateStateMachine),
 		FGameFeatureStateProgressUpdate::CreateUObject(this, &ThisClass::UpdateCurrentStateProgress));
 
@@ -1551,18 +1568,149 @@ void UGameFeaturePluginStateMachine::InitStateMachine(const FString& InPluginURL
 	AllStates[CurrentStateInfo.State]->BeginState();
 }
 
-void UGameFeaturePluginStateMachine::SetDestinationState(EGameFeaturePluginState InDestinationState, FGameFeatureStateTransitionComplete OnFeatureStateTransitionComplete)
+bool UGameFeaturePluginStateMachine::SetDestination(FGameFeaturePluginStateRange InDestination, FGameFeatureStateTransitionComplete OnFeatureStateTransitionComplete, FDelegateHandle* OutCallbackHandle /*= nullptr*/)
 {
-	check(IsValidDestinationState(InDestinationState));
+	check(IsValidDestinationState(InDestination.MinState));
+	check(IsValidDestinationState(InDestination.MaxState));
+
+	if (!InDestination.IsValid())
+	{
+		// Invalid range
+		return false;
+	}
 
 	// JMarcus TODO: If we aren't in a destination state and our new destination is in the opposite direction of 
 	// our current destination, cancel the current state transition (if possible)
-	// The completion delegate may be stomped in these cases.  Should probably callback with a cancelled error
+	// Should callback with a cancelled error
+
+	if (!IsRunning())
+	{
+		// Not running so any new range is acceptable
+
+		if (OutCallbackHandle)
+		{
+			OutCallbackHandle->Reset();
+		}
 
-	StateProperties.DestinationState = InDestinationState;
-	StateProperties.OnFeatureStateTransitionComplete = OnFeatureStateTransitionComplete;
+		FDestinationGameFeaturePluginState* CurrState = AllStates[CurrentStateInfo.State]->AsDestinationState();
 
-	UpdateStateMachine();
+		if (InDestination.Contains(CurrentStateInfo.State))
+		{
+			OnFeatureStateTransitionComplete.ExecuteIfBound(this, MakeValue());
+			return true;
+		}
+		
+		if (CurrentStateInfo.State < InDestination)
+		{
+			FDestinationGameFeaturePluginState* MinDestState = AllStates[InDestination.MinState]->AsDestinationState();
+			FDelegateHandle CallbackHandle = MinDestState->OnDestinationStateReached.Add(MoveTemp(OnFeatureStateTransitionComplete));
+			if (OutCallbackHandle)
+			{
+				*OutCallbackHandle = CallbackHandle;
+			}
+		}
+		else if (CurrentStateInfo.State > InDestination)
+		{
+			FDestinationGameFeaturePluginState* MaxDestState = AllStates[InDestination.MaxState]->AsDestinationState();
+			FDelegateHandle CallbackHandle = MaxDestState->OnDestinationStateReached.Add(MoveTemp(OnFeatureStateTransitionComplete));
+			if (OutCallbackHandle)
+			{
+				*OutCallbackHandle = CallbackHandle;
+			}
+		}
+
+		StateProperties.Destination = InDestination;
+		UpdateStateMachine();
+
+		return true;
+	}
+
+	if (TOptional<FGameFeaturePluginStateRange> NewDestination = InDestination.Intersect(InDestination))
+	{
+		// The machine is already running so we can only transition to this range if it overlaps with our current range.
+		// We can satisfy both ranges in this case.
+
+		if (OutCallbackHandle)
+		{
+			OutCallbackHandle->Reset();
+		}
+
+		if (CurrentStateInfo.State < StateProperties.Destination)
+		{
+			StateProperties.Destination = *NewDestination;
+
+			if (InDestination.Contains(CurrentStateInfo.State))
+			{
+				OnFeatureStateTransitionComplete.ExecuteIfBound(this, MakeValue());
+				return true;
+			}
+
+			FDestinationGameFeaturePluginState* MinDestState = AllStates[InDestination.MinState]->AsDestinationState();
+			FDelegateHandle CallbackHandle = MinDestState->OnDestinationStateReached.Add(MoveTemp(OnFeatureStateTransitionComplete));
+			if (OutCallbackHandle)
+			{
+				*OutCallbackHandle = CallbackHandle;
+			}
+		}
+		else if(CurrentStateInfo.State > StateProperties.Destination)
+		{
+			StateProperties.Destination = *NewDestination;
+
+			if (InDestination.Contains(CurrentStateInfo.State))
+			{
+				OnFeatureStateTransitionComplete.ExecuteIfBound(this, MakeValue());
+				return true;
+			}
+
+			FDestinationGameFeaturePluginState* MaxDestState = AllStates[InDestination.MaxState]->AsDestinationState();
+			FDelegateHandle CallbackHandle = MaxDestState->OnDestinationStateReached.Add(MoveTemp(OnFeatureStateTransitionComplete));
+			if (OutCallbackHandle)
+			{
+				*OutCallbackHandle = CallbackHandle;
+			}
+		}
+		else
+		{
+			checkf(false, TEXT("IsRunning() returned true but state machine has reached destination!"));
+		}
+
+		return true;
+	}
+
+	// The requested state range is completely outside the the current state range so reject the request
+	return false;
+}
+
+void UGameFeaturePluginStateMachine::RemovePendingTransitionCallback(FDelegateHandle InHandle)
+{
+	for (std::underlying_type<EGameFeaturePluginState>::type iState = 0;
+		iState < EGameFeaturePluginState::MAX;
+		++iState)
+	{
+		if (FDestinationGameFeaturePluginState* DestState = AllStates[iState]->AsDestinationState())
+		{
+			if (DestState->OnDestinationStateReached.Remove(InHandle))
+			{
+				break;
+			}
+		}
+	}
+}
+
+void UGameFeaturePluginStateMachine::RemovePendingTransitionCallback(void* DelegateObject)
+{
+	for (std::underlying_type<EGameFeaturePluginState>::type iState = 0;
+		iState < EGameFeaturePluginState::MAX;
+		++iState)
+	{
+		if (FDestinationGameFeaturePluginState* DestState = AllStates[iState]->AsDestinationState())
+		{
+			if (DestState->OnDestinationStateReached.RemoveAll(DelegateObject))
+			{
+				break;
+			}
+		}
+	}
 }
 
 FString UGameFeaturePluginStateMachine::GetGameFeatureName() const
@@ -1599,9 +1747,9 @@ EGameFeaturePluginState UGameFeaturePluginStateMachine::GetCurrentState() const
 	return GetCurrentStateInfo().State;
 }
 
-EGameFeaturePluginState UGameFeaturePluginStateMachine::GetDestinationState() const
+FGameFeaturePluginStateRange UGameFeaturePluginStateMachine::GetDestination() const
 {
-	return StateProperties.DestinationState;
+	return StateProperties.Destination;
 }
 
 const FGameFeaturePluginStateInfo& UGameFeaturePluginStateMachine::GetCurrentStateInfo() const
@@ -1609,6 +1757,11 @@ const FGameFeaturePluginStateInfo& UGameFeaturePluginStateMachine::GetCurrentSta
 	return CurrentStateInfo;
 }
 
+bool UGameFeaturePluginStateMachine::IsRunning() const
+{
+	return !StateProperties.Destination.Contains(CurrentStateInfo.State);
+}
+
 bool UGameFeaturePluginStateMachine::IsStatusKnown() const
 {
 	return GetCurrentState() == EGameFeaturePluginState::ErrorUnavailable || GetCurrentState() >= EGameFeaturePluginState::StatusKnown;
@@ -1667,9 +1820,8 @@ void UGameFeaturePluginStateMachine::UpdateStateMachine()
 		return;
 	}
 
-	TGuardValue<bool> ScopeGuard(bInUpdateStateMachine, true);
-
-	UE::GameFeatures::FResult TransitionResult(MakeValue());
+	TOptional<TGuardValue<bool>> ScopeGuard(InPlace, bInUpdateStateMachine, true);
+	
 	bool bKeepProcessing = false;
 	int32 NumTransitions = 0;
 	const int32 MaxTransitions = 10000;
@@ -1680,8 +1832,6 @@ void UGameFeaturePluginStateMachine::UpdateStateMachine()
 		FGameFeaturePluginStateStatus StateStatus;
 		AllStates[CurrentState]->UpdateState(StateStatus);
 
-		TransitionResult = StateStatus.TransitionResult;
-
 		if (StateStatus.TransitionToState == CurrentState)
 		{
 			UE_LOG(LogGameFeatures, Fatal, TEXT("Game feature state %s transitioning to itself. GameFeature: %s"), *UE::GameFeatures::ToString(CurrentState), *GetGameFeatureName());
@@ -1695,15 +1845,41 @@ void UGameFeaturePluginStateMachine::UpdateStateMachine()
 			CurrentState = StateStatus.TransitionToState;
 			check(CurrentState != EGameFeaturePluginState::MAX);
 			AllStates[CurrentState]->BeginState();
-			OnStateChangedEvent.Broadcast(this);
-			bKeepProcessing = true;
-		}
+			
+			const bool bError = !StateStatus.TransitionResult.HasValue();
+			if (bError)
+			{
+				check(IsValidErrorState(CurrentState));
+				StateProperties.Destination = FGameFeaturePluginStateRange(CurrentState);
+			}
 
-		if (!TransitionResult.HasValue())
-		{
-			check(IsValidErrorState(CurrentState));
-			StateProperties.DestinationState = CurrentState;
-			break;
+			bKeepProcessing = AllStates[CurrentState]->GetStateType() == EGameFeaturePluginStateType::Transition || !StateProperties.Destination.Contains(CurrentState);
+			if (!bKeepProcessing)
+			{
+				// Make sure bInUpdateStateMachine is not set while processing callbacks if we are at our destination
+				ScopeGuard.Reset();
+			}
+
+			// In case of an error, callback all possible callbacks, otherwise only for the current state
+			const FGameFeaturePluginStateRange CallbackRange = bError ? 
+				FGameFeaturePluginStateRange(EGameFeaturePluginState::Uninitialized, EGameFeaturePluginState::Active) :
+				FGameFeaturePluginStateRange(CurrentState);
+
+			for (std::underlying_type<EGameFeaturePluginState>::type iState = CallbackRange.MinState; 
+				iState <= CallbackRange.MaxState; 
+				++iState)
+			{
+				if (FDestinationGameFeaturePluginState* DestState = AllStates[iState]->AsDestinationState())
+				{
+					// Use a local callback on the stack. If SetDestination() is called from the callback then we don't want to stomp the callback
+					// for the new state transition request.
+					// Callback from terminal state could also trigger a GC that would destroy the state machine
+					FDestinationGameFeaturePluginState::FOnDestinationStateReached LocalOnDestinationStateReached(MoveTemp(DestState->OnDestinationStateReached));
+					DestState->OnDestinationStateReached.Clear();
+
+					LocalOnDestinationStateReached.Broadcast(this, StateStatus.TransitionResult);
+				}
+			}
 		}
 
 		if (NumTransitions++ > MaxTransitions)
@@ -1711,16 +1887,6 @@ void UGameFeaturePluginStateMachine::UpdateStateMachine()
 			UE_LOG(LogGameFeatures, Fatal, TEXT("Infinite loop in game feature state machine transitions. Current state %s. GameFeature: %s"), *UE::GameFeatures::ToString(CurrentState), *GetGameFeatureName());
 		}
 	} while (bKeepProcessing);
-
-	if (CurrentState == StateProperties.DestinationState)
-	{
-		check(IsValidTransitionState(CurrentState) == false);
-		// @note Clear out StateProperties.OnFeatureStateTransitionComplete prior to executing the 
-		// (termination) callback in case it triggers a GC that would destroy the state machine
-		FGameFeatureStateTransitionComplete TransitionCompleteCallback = MoveTemp(StateProperties.OnFeatureStateTransitionComplete);
-		StateProperties.OnFeatureStateTransitionComplete = FGameFeatureStateTransitionComplete();
-		TransitionCompleteCallback.ExecuteIfBound(this, TransitionResult);
-	}
 }
 
 void UGameFeaturePluginStateMachine::UpdateCurrentStateProgress(float Progress)
@@ -1730,13 +1896,11 @@ void UGameFeaturePluginStateMachine::UpdateCurrentStateProgress(float Progress)
 
 FGameFeaturePluginStateMachineProperties::FGameFeaturePluginStateMachineProperties(
 	const FString& InPluginURL,
-	EGameFeaturePluginState DesiredDestination,
-	const FGameFeaturePluginRequestStateMachineDependencies& RequestStateMachineDependenciesDelegate,
+	const FGameFeaturePluginStateRange& DesiredDestination,
 	const FGameFeaturePluginRequestUpdateStateMachine& RequestUpdateStateMachineDelegate,
 	const FGameFeatureStateProgressUpdate& FeatureStateProgressUpdateDelegate)
 	: PluginURL(InPluginURL)
-	, DestinationState(DesiredDestination)
-	, OnRequestStateMachineDependencies(RequestStateMachineDependenciesDelegate)
+	, Destination(DesiredDestination)
 	, OnRequestUpdateStateMachine(RequestUpdateStateMachineDelegate)
 	, OnFeatureStateProgressUpdate(FeatureStateProgressUpdateDelegate)
 {
@@ -1751,7 +1915,7 @@ EGameFeaturePluginProtocol FGameFeaturePluginStateMachineProperties::GetPluginPr
 
 	for (EGameFeaturePluginProtocol Proto : TEnumRange<EGameFeaturePluginProtocol>())
 	{
-		const TCHAR* Prefix = GameFeaturePluginProtocolPrefix(Proto);
+		const TCHAR* Prefix = UE::GameFeatures::GameFeaturePluginProtocolPrefix(Proto);
 		if (Prefix && *Prefix && PluginURL.StartsWith(Prefix))
 		{
 			CachedPluginProtocol = Proto;
@@ -1767,11 +1931,11 @@ bool FGameFeaturePluginStateMachineProperties::ParseURL()
 	if (GetPluginProtocol() == EGameFeaturePluginProtocol::File)
 	{
 		PluginInstalledFilename = PluginURL.RightChop(
-			FCString::Strlen(GameFeaturePluginProtocolPrefix(EGameFeaturePluginProtocol::File)));
+			FCString::Strlen(UE::GameFeatures::GameFeaturePluginProtocolPrefix(EGameFeaturePluginProtocol::File)));
 	}
 	else if (GetPluginProtocol() == EGameFeaturePluginProtocol::InstallBundle)
 	{
-		int32 CursorIdx = FCString::Strlen(GameFeaturePluginProtocolPrefix(EGameFeaturePluginProtocol::InstallBundle));
+		int32 CursorIdx = FCString::Strlen(UE::GameFeatures::GameFeaturePluginProtocolPrefix(EGameFeaturePluginProtocol::InstallBundle));
 		int32 QueryIdx = PluginURL.Find(TEXT("?"), ESearchCase::IgnoreCase, ESearchDir::FromStart, CursorIdx);
 		if (QueryIdx == INDEX_NONE)
 		{
@@ -1789,7 +1953,7 @@ bool FGameFeaturePluginStateMachineProperties::ParseURL()
 			return false;
 		}
 
-		FInstallBundlePluginProtocolMetaData& MetaData = *ProtocolMetadata.SetSubtype<FInstallBundlePluginProtocolMetaData>();
+		UE::GameFeatures::FInstallBundlePluginProtocolMetaData& MetaData = *ProtocolMetadata.SetSubtype<UE::GameFeatures::FInstallBundlePluginProtocolMetaData>();
 		MetaData.InstallBundles.Reserve(BundleNames.Num());
 		for (FString& BundleNameString : BundleNames)
 		{
diff --git a/Engine/Plugins/Experimental/GameFeatures/Source/GameFeatures/Private/GameFeaturePluginStateMachine.h b/Engine/Plugins/Experimental/GameFeatures/Source/GameFeatures/Private/GameFeaturePluginStateMachine.h
--- a/Engine/Plugins/Experimental/GameFeatures/Source/GameFeatures/Private/GameFeaturePluginStateMachine.h
+++ b/Engine/Plugins/Experimental/GameFeatures/Source/GameFeatures/Private/GameFeaturePluginStateMachine.h
@@ -128,15 +128,80 @@ Transition states are expected to transition the machine to another state after
                +------------+
 */
 
-// JMarcus TODO: move this stuff into UE::GameFeatures namespace
+namespace UE::GameFeatures
+{
+	const TCHAR* GameFeaturePluginProtocolPrefix(EGameFeaturePluginProtocol Protocol);
 
-const TCHAR* GameFeaturePluginProtocolPrefix(EGameFeaturePluginProtocol Protocol);
+	struct FInstallBundlePluginProtocolMetaData
+	{
+		TArray<FName> InstallBundles;
+	};
+}
 
-struct FInstallBundlePluginProtocolMetaData
+struct FGameFeaturePluginStateRange
 {
-	TArray<FName> InstallBundles;
+	EGameFeaturePluginState MinState = EGameFeaturePluginState::Uninitialized;
+	EGameFeaturePluginState MaxState = EGameFeaturePluginState::Uninitialized;
+
+	FGameFeaturePluginStateRange() = default;
+
+	FGameFeaturePluginStateRange(EGameFeaturePluginState InMinState, EGameFeaturePluginState InMaxState)
+		: MinState(InMinState), MaxState(InMaxState)
+	{}
+
+	explicit FGameFeaturePluginStateRange(EGameFeaturePluginState InState)
+		: MinState(InState), MaxState(InState)
+	{}
+
+	bool IsValid() const { return MinState <= MaxState; }
+
+	bool Contains(EGameFeaturePluginState InState) const
+	{
+		return InState >= MinState && InState <= MaxState;
+	}
+
+	bool Overlaps(const FGameFeaturePluginStateRange& Other) const
+	{
+		return Other.MinState <= MaxState && Other.MaxState >= MinState;
+	}
+
+	TOptional<FGameFeaturePluginStateRange> Intersect(const FGameFeaturePluginStateRange& Other) const
+	{
+		TOptional<FGameFeaturePluginStateRange> Intersection;
+
+		if (Overlaps(Other))
+		{
+			Intersection.Emplace(FMath::Max(Other.MinState, MinState), FMath::Min(Other.MaxState, MaxState));
+		}
+
+		return Intersection;
+	}
+
+	bool operator==(const FGameFeaturePluginStateRange& Other) const { return MinState == Other.MinState && MaxState == Other.MaxState; }
+	bool operator<(const FGameFeaturePluginStateRange& Other) const { return MaxState < Other.MinState; }
+	bool operator>(const FGameFeaturePluginStateRange& Other) const { return MinState < Other.MaxState; }
 };
 
+inline bool operator<(EGameFeaturePluginState State, const FGameFeaturePluginStateRange& StateRange)
+{
+	return State < StateRange.MinState;
+}
+
+inline bool operator<(const FGameFeaturePluginStateRange& StateRange, EGameFeaturePluginState State)
+{
+	return StateRange.MaxState < State;
+}
+
+inline bool operator>(EGameFeaturePluginState State, const FGameFeaturePluginStateRange& StateRange)
+{
+	return State > StateRange.MaxState;
+}
+
+inline bool operator>(const FGameFeaturePluginStateRange& StateRange, EGameFeaturePluginState State)
+{
+	return StateRange.MinState > State;
+}
+
 /** Notification that a state transition is complete */
 DECLARE_DELEGATE_TwoParams(FGameFeatureStateTransitionComplete, UGameFeaturePluginStateMachine* /*Machine*/, const UE::GameFeatures::FResult& /*Result*/);
 
@@ -169,23 +234,20 @@ struct FGameFeaturePluginStateMachineProperties
 	FString PluginName;
 
 	/** Meta data parsed from the URL for a specific protocol. */
-	TUnion<FInstallBundlePluginProtocolMetaData> ProtocolMetadata;
+	TUnion<UE::GameFeatures::FInstallBundlePluginProtocolMetaData> ProtocolMetadata;
 
 	/** Tracks whether or not this state machine added the pluging to the plugin manager. */
 	bool bAddedPluginToManager = false;
 
 	/** The desired state during a transition. */
-	EGameFeaturePluginState DestinationState = EGameFeaturePluginState::Uninitialized;
+	FGameFeaturePluginStateRange Destination;
 
 	/** The data asset describing this game feature */
 	UPROPERTY(Transient)
 	UGameFeatureData* GameFeatureData = nullptr;
 
 	/** Delegate for when a state transition request has completed. */
-	FGameFeatureStateTransitionComplete OnFeatureStateTransitionComplete;
-
-	/** Delegate for when this machine needs to request its dependency machines. */
-	FGameFeaturePluginRequestStateMachineDependencies OnRequestStateMachineDependencies;
+	FGameFeatureStateTransitionComplete OnFeatureStateTransitionComplete;	
 
 	/** Delegate to request the state machine be updated. */
 	FGameFeaturePluginRequestUpdateStateMachine OnRequestUpdateStateMachine;
@@ -196,8 +258,7 @@ struct FGameFeaturePluginStateMachineProperties
 	FGameFeaturePluginStateMachineProperties() = default;
 	FGameFeaturePluginStateMachineProperties(
 		const FString& InPluginURL,
-		EGameFeaturePluginState DesiredDestination,
-		const FGameFeaturePluginRequestStateMachineDependencies& RequestStateMachineDependenciesDelegate,
+		const FGameFeaturePluginStateRange& DesiredDestination,
 		const FGameFeaturePluginRequestUpdateStateMachine& RequestUpdateStateMachineDelegate,
 		const FGameFeatureStateProgressUpdate& FeatureStateProgressUpdateDelegate);
 
@@ -242,6 +303,9 @@ enum class EGameFeaturePluginStateType : uint8
 	Error
 };
 
+struct FDestinationGameFeaturePluginState;
+struct FErrorGameFeaturePluginState;
+
 /** Base class for all game feature plugin states */
 struct FGameFeaturePluginState
 {
@@ -260,6 +324,9 @@ struct FGameFeaturePluginState
 	/** Returns the type of state this is */
 	virtual EGameFeaturePluginStateType GetStateType() const { return EGameFeaturePluginStateType::Transition; }
 
+	FDestinationGameFeaturePluginState* AsDestinationState();
+	FErrorGameFeaturePluginState* AsErrorState();
+
 	/** The common properties that can be accessed by the states of the state machine */
 	FGameFeaturePluginStateMachineProperties& StateProperties;
 
@@ -275,6 +342,27 @@ private:
 	mutable FTSTicker::FDelegateHandle TickHandle;
 };
 
+/** Base class for destination game feature plugin states */
+struct FDestinationGameFeaturePluginState : public FGameFeaturePluginState
+{
+	FDestinationGameFeaturePluginState(FGameFeaturePluginStateMachineProperties& InStateProperties) : FGameFeaturePluginState(InStateProperties) {}
+
+	/** Returns the type of state this is */
+	virtual EGameFeaturePluginStateType GetStateType() const override { return EGameFeaturePluginStateType::Destination; }
+
+	DECLARE_MULTICAST_DELEGATE_TwoParams(FOnDestinationStateReached, UGameFeaturePluginStateMachine* /*Machine*/, const UE::GameFeatures::FResult& /*Result*/);
+	FOnDestinationStateReached OnDestinationStateReached;
+};
+
+/** Base class for error game feature plugin states */
+struct FErrorGameFeaturePluginState : public FDestinationGameFeaturePluginState
+{
+	FErrorGameFeaturePluginState(FGameFeaturePluginStateMachineProperties& InStateProperties) : FDestinationGameFeaturePluginState(InStateProperties) {}
+
+	/** Returns the type of state this is */
+	virtual EGameFeaturePluginStateType GetStateType() const override { return EGameFeaturePluginStateType::Error; }
+};
+
 /** Information about a given plugin state, used to expose information to external code */
 struct FGameFeaturePluginStateInfo
 {
@@ -298,10 +386,19 @@ public:
 	UGameFeaturePluginStateMachine(const FObjectInitializer& ObjectInitializer);
 
 	/** Initializes the state machine and assigns the URL for the plugin it manages. This sets the machine to the 'UnknownStatus' state. */
-	void InitStateMachine(const FString& InPluginURL, const FGameFeaturePluginRequestStateMachineDependencies& OnRequestStateMachineDependencies);
+	void InitStateMachine(const FString& InPluginURL);
 
-	/** Asynchronously transitions the state machine to the destination state and reports when it is done. DestinationState must not be a transition state like 'Downloading' */
-	void SetDestinationState(EGameFeaturePluginState InDestinationState, FGameFeatureStateTransitionComplete OnFeatureStateTransitionComplete);
+	/** Asynchronously transitions the state machine to the destination state range and reports when it is done. 
+	  * DestinationState must be of type EGameFeaturePluginStateType::Destination.
+	  * If returns true and OnFeatureStateTransitionComplete is not called immediately, OutCallbackHandle will be set
+	  * Returns false and does not callback if a transition is already in progress and the destination range is not compatible with the current range. */
+	bool SetDestination(FGameFeaturePluginStateRange InDestination, FGameFeatureStateTransitionComplete OnFeatureStateTransitionComplete, FDelegateHandle* OutCallbackHandle = nullptr);
+
+	/** Remove any pending callback from SetDestination */
+	void RemovePendingTransitionCallback(FDelegateHandle InHandle);
+
+	/** Remove any pending callback from SetDestination */
+	void RemovePendingTransitionCallback(void* DelegateObject);
 
 	/** Returns the name of the game feature. Before StatusKnown, this returns the URL. */
 	FString GetGameFeatureName() const;
@@ -318,12 +415,15 @@ public:
 	/** Returns the enum state for this machine */
 	EGameFeaturePluginState GetCurrentState() const;
 
-	/** Returns the state this machine is trying to move to */
-	EGameFeaturePluginState GetDestinationState() const;
+	/** Returns the state range this machine is trying to move to */
+	FGameFeaturePluginStateRange GetDestination() const;
 
 	/** Returns information about the current state */
 	const FGameFeaturePluginStateInfo& GetCurrentStateInfo() const;
 
+	/** Returns true if attempting to reach a new destination state */
+	bool IsRunning() const;
+
 	/** Returns true if the state is at least StatusKnown so we can query info about the game feature plugin */
 	bool IsStatusKnown() const;
 
@@ -336,10 +436,6 @@ public:
 	/** If the plugin is registered already, we will retrieve its game feature data */
 	UGameFeatureData* GetGameFeatureDataForRegisteredPlugin();
 
-	/** Delegate for the machine's state changed. */
-	DECLARE_EVENT_OneParam(UGameFeaturePluginStateMachine, FGameFeaturePluginStateChanged, UGameFeaturePluginStateMachine* /*Machine*/);
-	FGameFeaturePluginStateChanged& OnStateChanged() { return OnStateChangedEvent; }
-
 private:
 	/** Returns true if the specified state is not a transition state */
 	bool IsValidTransitionState(EGameFeaturePluginState InState) const;
@@ -356,9 +452,6 @@ private:
 	/** Update Progress for current state */
 	void UpdateCurrentStateProgress(float Progress);
 
-	/** Delegate for the machine's state changed. */
-	FGameFeaturePluginStateChanged OnStateChangedEvent;
-
 	/** Information about the current state */
 	FGameFeaturePluginStateInfo CurrentStateInfo;
 
diff --git a/Engine/Plugins/Experimental/GameFeatures/Source/GameFeatures/Private/GameFeaturesSubsystem.cpp b/Engine/Plugins/Experimental/GameFeatures/Source/GameFeatures/Private/GameFeaturesSubsystem.cpp
--- a/Engine/Plugins/Experimental/GameFeatures/Source/GameFeatures/Private/GameFeaturesSubsystem.cpp
+++ b/Engine/Plugins/Experimental/GameFeatures/Source/GameFeatures/Private/GameFeaturesSubsystem.cpp
@@ -380,7 +380,7 @@ EGameFeaturePluginProtocol UGameFeaturesSubsystem::GetPluginURLProtocol(FStringV
 {
 	for (EGameFeaturePluginProtocol Protocol : TEnumRange<EGameFeaturePluginProtocol>())
 	{
-		if (PluginURL.StartsWith(GameFeaturePluginProtocolPrefix(Protocol)))
+		if (PluginURL.StartsWith(UE::GameFeatures::GameFeaturePluginProtocolPrefix(Protocol)))
 		{
 			return Protocol;
 		}
@@ -575,18 +575,36 @@ bool UGameFeaturesSubsystem::IsGameFeaturePluginLoaded(const FString& PluginURL)
 
 void UGameFeaturesSubsystem::LoadGameFeaturePlugin(const FString& PluginURL, const FGameFeaturePluginLoadComplete& CompleteDelegate)
 {
-	if (IsGameFeaturePluginLoaded(PluginURL))
+	const bool bIsPluginAllowed = GameSpecificPolicies->IsPluginAllowed(PluginURL);
+	if (!bIsPluginAllowed)
 	{
-		FTSTicker::GetCoreTicker().AddTicker(FTickerDelegate::CreateWeakLambda(this, [this, CompleteDelegate](float dts)
-			{
-				CompleteDelegate.ExecuteIfBound(MakeValue());
-				return false;
-			}));
+		CompleteDelegate.ExecuteIfBound(MakeError(TEXT("GameFeaturePlugin.StateMachine.Plugin_Denied_By_GameSpecificPolicy")));
+		return;
+	}
+
+	UGameFeaturePluginStateMachine* StateMachine = FindOrCreateGameFeaturePluginStateMachine(PluginURL);
+	
+	if (!StateMachine->IsRunning() && StateMachine->GetCurrentState() == EGameFeaturePluginState::Active)
+	{
+		// TODO: Resolve the activated case here, this is needed because in a PIE environment the plugins
+		// are not sandboxed, and we need to do simulate a successful activate call in order run GFP systems 
+		// on whichever Role runs second between client and server.
 
-		return; // Early out, we are already loaded.
+		// Refire the observer for Activated and do nothing else.
+		for (UObject* Observer : Observers)
+		{
+			CastChecked<IGameFeatureStateChangeObserver>(Observer)->OnGameFeatureActivating(StateMachine->GetGameFeatureDataForActivePlugin(), PluginURL);
+		}
 	}
 
-	ChangeGameFeatureTargetState(PluginURL, EGameFeatureTargetState::Loaded, CompleteDelegate);
+	const bool bSetDestination = StateMachine->SetDestination(
+		FGameFeaturePluginStateRange(EGameFeaturePluginState::Loaded, EGameFeaturePluginState::Active),
+		FGameFeatureStateTransitionComplete::CreateUObject(this, &ThisClass::ChangeGameFeatureTargetStateComplete, CompleteDelegate));
+
+	if (!bSetDestination)
+	{
+		CompleteDelegate.ExecuteIfBound(MakeError(TEXT("GameFeaturePlugin.StateMachine.State_Currently_Unreachable")));
+	}
 }
 
 void UGameFeaturesSubsystem::LoadAndActivateGameFeaturePlugin(const FString& PluginURL, const FGameFeaturePluginLoadComplete& CompleteDelegate)
@@ -597,61 +615,66 @@ void UGameFeaturesSubsystem::LoadAndActivateGameFeaturePlugin(const FString& Plu
 void UGameFeaturesSubsystem::ChangeGameFeatureTargetState(const FString& PluginURL, EGameFeatureTargetState TargetState, const FGameFeaturePluginChangeStateComplete& CompleteDelegate)
 {
 	EGameFeaturePluginState TargetPluginState = EGameFeaturePluginState::MAX;
-	if (TargetState == EGameFeatureTargetState::Active)
+
+	switch (TargetState)
 	{
-		TargetPluginState = EGameFeaturePluginState::Active;
+	case EGameFeatureTargetState::Installed:	TargetPluginState = EGameFeaturePluginState::Installed;		break;
+	case EGameFeatureTargetState::Registered:	TargetPluginState = EGameFeaturePluginState::Registered;	break;
+	case EGameFeatureTargetState::Loaded:		TargetPluginState = EGameFeaturePluginState::Loaded;		break;
+	case EGameFeatureTargetState::Active:		TargetPluginState = EGameFeaturePluginState::Active;		break;
 	}
-	else if (TargetState == EGameFeatureTargetState::Loaded)
+
+	// Make sure we have coverage on all values of EGameFeatureTargetState
+	static_assert(std::underlying_type<EGameFeatureTargetState>::type(EGameFeatureTargetState::Count) == 4, "");
+	check(TargetPluginState != EGameFeaturePluginState::MAX);
+
+	const bool bIsPluginAllowed = GameSpecificPolicies->IsPluginAllowed(PluginURL);
+
+	UGameFeaturePluginStateMachine* StateMachine = nullptr;
+	if (!bIsPluginAllowed)
 	{
-		TargetPluginState = EGameFeaturePluginState::Loaded;
+		StateMachine = FindGameFeaturePluginStateMachine(PluginURL);
+		if (!StateMachine)
+		{
+			CompleteDelegate.ExecuteIfBound(MakeError(TEXT("GameFeaturePlugin.StateMachine.Plugin_Denied_By_GameSpecificPolicy")));
+			return;
+		}
 	}
-	else if (TargetState == EGameFeatureTargetState::Registered)
+	else
 	{
-		TargetPluginState = EGameFeaturePluginState::Registered;
+		StateMachine = FindOrCreateGameFeaturePluginStateMachine(PluginURL);
 	}
-	else if (TargetState == EGameFeatureTargetState::Installed)
+	
+	check(StateMachine);
+
+	if (TargetPluginState > StateMachine->GetCurrentState() && !GameSpecificPolicies->IsPluginAllowed(PluginURL))
 	{
-		TargetPluginState = EGameFeaturePluginState::Installed;
+		CompleteDelegate.ExecuteIfBound(MakeError(TEXT("GameFeaturePlugin.StateMachine.Plugin_Denied_By_GameSpecificPolicy")));
+		return;
 	}
 
-	if (TargetPluginState != EGameFeaturePluginState::MAX)
+	if (TargetState == EGameFeatureTargetState::Active &&
+		!StateMachine->IsRunning() && 
+		StateMachine->GetCurrentState() == TargetPluginState)
 	{
-		if (UGameFeaturePluginStateMachine* StateMachine = FindOrCreateGameFeaturePluginStateMachine(PluginURL))
-		{
-			if (TargetState == EGameFeatureTargetState::Active
-				&& StateMachine->GetCurrentState() == TargetPluginState)
-			{
-				// TODO: Resolve the activated case here, this is needed because in a PIE environment the plugins
-				// are not sandboxed, and we need to do simulate a successful activate call in order run GFP systems 
-				// on whichever Role runs second between client and server.
-
-				// Refire the observer for Activated and do nothing else.
-				for (UObject* Observer : Observers)
-				{
-					CastChecked<IGameFeatureStateChangeObserver>(Observer)->OnGameFeatureActivating(StateMachine->GetGameFeatureDataForActivePlugin(), PluginURL);
-				}
+		// TODO: Resolve the activated case here, this is needed because in a PIE environment the plugins
+		// are not sandboxed, and we need to do simulate a successful activate call in order run GFP systems 
+		// on whichever Role runs second between client and server.
 
-				FTSTicker::GetCoreTicker().AddTicker(FTickerDelegate::CreateWeakLambda(this, [CompleteDelegate](float)
-					{
-						CompleteDelegate.ExecuteIfBound(UE::GameFeatures::FResult(MakeValue()));
-						return false;
-					}));
-			}
-			else if (TargetPluginState > StateMachine->GetCurrentState()
-				&& !GameSpecificPolicies->IsPluginAllowed(PluginURL))
-			{
-				FTSTicker::GetCoreTicker().AddTicker(FTickerDelegate::CreateWeakLambda(this, [CompleteDelegate](float)
-					{
-						CompleteDelegate.ExecuteIfBound(MakeError(TEXT("GameFeaturePlugin.StateMachine.Plugin_Denied_By_GameSpecificPolicy")));
-						return false;
-					}));
-			}
-			else
-			{
-				StateMachine->SetDestinationState(TargetPluginState, FGameFeatureStateTransitionComplete::CreateUObject(this, &ThisClass::ChangeGameFeatureTargetStateComplete, CompleteDelegate));
-			}
+		// Refire the observer for Activated and do nothing else.
+		for (UObject* Observer : Observers)
+		{
+			CastChecked<IGameFeatureStateChangeObserver>(Observer)->OnGameFeatureActivating(StateMachine->GetGameFeatureDataForActivePlugin(), PluginURL);
 		}
 	}
+	
+	const bool bSetDestination = StateMachine->SetDestination(FGameFeaturePluginStateRange(TargetPluginState), 
+		FGameFeatureStateTransitionComplete::CreateUObject(this, &ThisClass::ChangeGameFeatureTargetStateComplete, CompleteDelegate));
+
+	if (!bSetDestination)
+	{
+		CompleteDelegate.ExecuteIfBound(MakeError(TEXT("GameFeaturePlugin.StateMachine.State_Currently_Unreachable")));
+	}
 }
 
 bool UGameFeaturesSubsystem::GetGameFeaturePluginInstallPercent(const FString& PluginURL, float& Install_Percent) const
@@ -665,7 +688,7 @@ bool UGameFeaturesSubsystem::GetGameFeaturePluginInstallPercent(const FString& P
 			{
 				Install_Percent = StateInfo.Progress;
 			}
-			else if (StateMachine->GetDestinationState() >= EGameFeaturePluginState::Installed && StateInfo.State >= EGameFeaturePluginState::Installed)
+			else if (StateInfo.State >= EGameFeaturePluginState::Installed)
 			{
 				Install_Percent = 1.0f;
 			}
@@ -693,152 +716,99 @@ bool UGameFeaturesSubsystem::IsGameFeaturePluginActive(const FString& PluginURL,
 
 void UGameFeaturesSubsystem::DeactivateGameFeaturePlugin(const FString& PluginURL)
 {
-	if (UGameFeaturePluginStateMachine* StateMachine = FindGameFeaturePluginStateMachine(PluginURL))
-	{
-		if (StateMachine->GetDestinationState() > EGameFeaturePluginState::Loaded)
-		{
-			FGameFeaturePluginDeactivateComplete Callback = FGameFeaturePluginDeactivateComplete();
-			DeactivateGameFeaturePlugin(PluginURL, Callback);
-		}
-	}
+	DeactivateGameFeaturePlugin(PluginURL, FGameFeaturePluginDeactivateComplete());
 }
 
 void UGameFeaturesSubsystem::DeactivateGameFeaturePlugin(const FString& PluginURL, const FGameFeaturePluginDeactivateComplete& CompleteDelegate)
 {
 	if (UGameFeaturePluginStateMachine* StateMachine = FindGameFeaturePluginStateMachine(PluginURL))
 	{
-		if (StateMachine->GetDestinationState() <= EGameFeaturePluginState::Loaded)
-		{
-			FTSTicker::GetCoreTicker().AddTicker(FTickerDelegate::CreateWeakLambda(this, [CompleteDelegate](float)
-				{
-					CompleteDelegate.ExecuteIfBound(UE::GameFeatures::FResult(MakeValue()));
-					return false;
-				}));
-		}
-		else
+		const bool bSetDestination = StateMachine->SetDestination(
+			FGameFeaturePluginStateRange(EGameFeaturePluginState::Terminal, EGameFeaturePluginState::Loaded), 
+			FGameFeatureStateTransitionComplete::CreateUObject(this, &ThisClass::ChangeGameFeatureTargetStateComplete, CompleteDelegate));
+
+		if (!bSetDestination)
 		{
-			ChangeGameFeatureTargetState(PluginURL, EGameFeatureTargetState::Loaded, CompleteDelegate);
+			CompleteDelegate.ExecuteIfBound(MakeError(TEXT("GameFeaturePlugin.StateMachine.State_Currently_Unreachable")));
 		}
 	}
 	else
 	{
-		FTSTicker::GetCoreTicker().AddTicker(FTickerDelegate::CreateWeakLambda(this, [CompleteDelegate](float)
-			{
-				CompleteDelegate.ExecuteIfBound(UE::GameFeatures::FResult(MakeError(TEXT("GameFeaturePlugin.BadURL"))));
-				return false;
-			}));
+		CompleteDelegate.ExecuteIfBound(UE::GameFeatures::FResult(MakeError(TEXT("GameFeaturePlugin.BadURL"))));
 	}
 }
 
-void UGameFeaturesSubsystem::UnloadGameFeaturePlugin(const FString& PluginURL, bool bKeepRegistered)
+void UGameFeaturesSubsystem::UnloadGameFeaturePlugin(const FString& PluginURL, bool bKeepRegistered /*= false*/)
 {
-	ChangeGameFeatureTargetState(PluginURL,
-		bKeepRegistered ? EGameFeatureTargetState::Registered : EGameFeatureTargetState::Installed,
-		FGameFeaturePluginUnloadComplete());
+	UnloadGameFeaturePlugin(PluginURL, FGameFeaturePluginUnloadComplete(), bKeepRegistered);
 }
 
-void UGameFeaturesSubsystem::UnloadGameFeaturePlugin(const FString& PluginURL, const FGameFeaturePluginUnloadComplete& CompleteDelegate, bool bKeepRegistered)
+void UGameFeaturesSubsystem::UnloadGameFeaturePlugin(const FString& PluginURL, const FGameFeaturePluginUnloadComplete& CompleteDelegate, bool bKeepRegistered /*= false*/)
 {
 	if (UGameFeaturePluginStateMachine* StateMachine = FindGameFeaturePluginStateMachine(PluginURL))
 	{
-		if (StateMachine->GetDestinationState() <= EGameFeaturePluginState::Loaded)
-		{
-			FTSTicker::GetCoreTicker().AddTicker(FTickerDelegate::CreateWeakLambda(this, [CompleteDelegate](float)
-				{
-					CompleteDelegate.ExecuteIfBound(UE::GameFeatures::FResult(MakeValue()));
-					return false;
-				}));
-		}
-		else
+		EGameFeaturePluginState TargetPluginState = bKeepRegistered ? EGameFeaturePluginState::Registered : EGameFeaturePluginState::Installed;
+
+		const bool bSetDestination = StateMachine->SetDestination(
+			FGameFeaturePluginStateRange(EGameFeaturePluginState::Terminal, TargetPluginState),
+			FGameFeatureStateTransitionComplete::CreateUObject(this, &ThisClass::ChangeGameFeatureTargetStateComplete, CompleteDelegate));
+
+		if (!bSetDestination)
 		{
-			ChangeGameFeatureTargetState(PluginURL,
-				bKeepRegistered ? EGameFeatureTargetState::Registered : EGameFeatureTargetState::Installed,
-				CompleteDelegate);
+			CompleteDelegate.ExecuteIfBound(MakeError(TEXT("GameFeaturePlugin.StateMachine.State_Currently_Unreachable")));
 		}
 	}
 	else
 	{
-		FTSTicker::GetCoreTicker().AddTicker(FTickerDelegate::CreateWeakLambda(this, [CompleteDelegate](float)
-			{
-				CompleteDelegate.ExecuteIfBound(UE::GameFeatures::FResult(MakeError(TEXT("GameFeaturePlugin.BadURL"))));
-				return false;
-			}));
+		CompleteDelegate.ExecuteIfBound(UE::GameFeatures::FResult(MakeError(TEXT("GameFeaturePlugin.BadURL"))));
 	}
 }
 
 void UGameFeaturesSubsystem::UninstallGameFeaturePlugin(const FString& PluginURL)
 {
-	if (UGameFeaturePluginStateMachine* StateMachine = FindGameFeaturePluginStateMachine(PluginURL))
-	{
-		if (StateMachine->GetDestinationState() > EGameFeaturePluginState::StatusKnown)
-		{
-			FGameFeaturePluginUninstallComplete Callback = FGameFeaturePluginUninstallComplete();
-			UninstallGameFeaturePlugin(PluginURL, Callback);
-		}
-	}
+	UninstallGameFeaturePlugin(PluginURL, FGameFeaturePluginUninstallComplete());
 }
 
 void UGameFeaturesSubsystem::UninstallGameFeaturePlugin(const FString& PluginURL, const FGameFeaturePluginUninstallComplete& CompleteDelegate)
 {
 	if (UGameFeaturePluginStateMachine* StateMachine = FindGameFeaturePluginStateMachine(PluginURL))
 	{
-		ensureAlwaysMsgf(StateMachine->GetCurrentState() == StateMachine->GetDestinationState(), TEXT("Setting a new destination state while state machine is running!"));
+		const bool bSetDestination = StateMachine->SetDestination(
+			FGameFeaturePluginStateRange(EGameFeaturePluginState::Terminal, EGameFeaturePluginState::StatusKnown),
+			FGameFeatureStateTransitionComplete::CreateUObject(this, &ThisClass::ChangeGameFeatureTargetStateComplete, CompleteDelegate));
 
-		if (StateMachine->GetCurrentState() > EGameFeaturePluginState::StatusKnown)
+		if (!bSetDestination)
 		{
-			StateMachine->SetDestinationState(EGameFeaturePluginState::StatusKnown, FGameFeatureStateTransitionComplete::CreateUObject(this, &ThisClass::ChangeGameFeatureTargetStateComplete, CompleteDelegate));
-		}
-		else if (StateMachine->GetCurrentState() <= EGameFeaturePluginState::StatusKnown)
-		{
-			FTSTicker::GetCoreTicker().AddTicker(FTickerDelegate::CreateWeakLambda(this, [this, CompleteDelegate](float dts)
-			{
-				CompleteDelegate.ExecuteIfBound(UE::GameFeatures::FResult(MakeValue()));
-				return false;
-			}));
+			CompleteDelegate.ExecuteIfBound(MakeError(TEXT("GameFeaturePlugin.StateMachine.State_Currently_Unreachable")));
 		}
 	}
 	else
 	{
-		FTSTicker::GetCoreTicker().AddTicker(FTickerDelegate::CreateWeakLambda(this, [this, CompleteDelegate](float dts)
-		{
-			CompleteDelegate.ExecuteIfBound(UE::GameFeatures::FResult(MakeError(TEXT("GameFeaturePlugin.BadURL"))));
-			return false;
-		}));
+		CompleteDelegate.ExecuteIfBound(UE::GameFeatures::FResult(MakeError(TEXT("GameFeaturePlugin.BadURL"))));
 	}
 }
 
 void UGameFeaturesSubsystem::TerminateGameFeaturePlugin(const FString& PluginURL)
 {
-	FGameFeaturePluginUninstallComplete Callback = FGameFeaturePluginUninstallComplete();
-	TerminateGameFeaturePlugin(PluginURL, Callback);
+	TerminateGameFeaturePlugin(PluginURL, FGameFeaturePluginUninstallComplete());
 }
 
 void UGameFeaturesSubsystem::TerminateGameFeaturePlugin(const FString& PluginURL, const FGameFeaturePluginUninstallComplete& CompleteDelegate)
 {
 	if (UGameFeaturePluginStateMachine* StateMachine = FindGameFeaturePluginStateMachine(PluginURL))
 	{
-		ensureAlwaysMsgf(StateMachine->GetCurrentState() == StateMachine->GetDestinationState(), TEXT("Setting a new destination state while state machine is running!"));
+		const bool bSetDestination = StateMachine->SetDestination(
+			FGameFeaturePluginStateRange(EGameFeaturePluginState::Terminal),
+			FGameFeatureStateTransitionComplete::CreateUObject(this, &ThisClass::TerminateGameFeaturePluginComplete, CompleteDelegate));
 
-		if (StateMachine->GetCurrentState() > EGameFeaturePluginState::Terminal)
-		{
-			StateMachine->SetDestinationState(EGameFeaturePluginState::Terminal, FGameFeatureStateTransitionComplete::CreateUObject(this, &ThisClass::TerminateGameFeaturePluginComplete, CompleteDelegate));
-		}
-		else if (StateMachine->GetCurrentState() <= EGameFeaturePluginState::Terminal)
+		if (!bSetDestination)
 		{
-			FTSTicker::GetCoreTicker().AddTicker(FTickerDelegate::CreateWeakLambda(this, [this, CompleteDelegate](float dts)
-			{
-				CompleteDelegate.ExecuteIfBound(UE::GameFeatures::FResult(MakeValue()));
-				return false;
-			}));
+			CompleteDelegate.ExecuteIfBound(MakeError(TEXT("GameFeaturePlugin.StateMachine.State_Currently_Unreachable")));
 		}
 	}
 	else
 	{
-		FTSTicker::GetCoreTicker().AddTicker(FTickerDelegate::CreateWeakLambda(this, [this, CompleteDelegate](float dts)
-		{
-			CompleteDelegate.ExecuteIfBound(UE::GameFeatures::FResult(MakeError(TEXT("GameFeaturePlugin.BadURL"))));
-			return false;
-		}));
+		CompleteDelegate.ExecuteIfBound(UE::GameFeatures::FResult(MakeError(TEXT("GameFeaturePlugin.BadURL"))));
 	}
 }
 
@@ -867,15 +837,9 @@ void UGameFeaturesSubsystem::LoadBuiltInGameFeaturePlugin(const TSharedRef<IPlug
 						
 					const EGameFeaturePluginState DestinationState = ConvertInitialFeatureStateToTargetState(InitialAutoState);
 
-					if (StateMachine->GetCurrentState() >= DestinationState)
-					{
-						// If we're already at the destination or beyond, don't transition back
-						LoadGameFeaturePluginComplete(StateMachine, MakeValue());
-					}
-					else
-					{
-						StateMachine->SetDestinationState(DestinationState, FGameFeatureStateTransitionComplete::CreateUObject(this, &ThisClass::LoadGameFeaturePluginComplete));
-					}
+					// If we're already at the destination or beyond, don't transition back
+					verify(StateMachine->SetDestination(FGameFeaturePluginStateRange(DestinationState, EGameFeaturePluginState::Active), 
+						FGameFeatureStateTransitionComplete::CreateUObject(this, &ThisClass::LoadGameFeaturePluginComplete)));
 				}
 			}
 		}
@@ -971,6 +935,10 @@ EGameFeaturePluginState UGameFeaturesSubsystem::GetPluginState(const FString& Pl
 
 bool UGameFeaturesSubsystem::GetGameFeaturePluginDetails(const FString& PluginDescriptorFilename, FGameFeaturePluginDetails& OutPluginDetails) const
 {
+	// @TODO: We load the descriptor 2-3 per plugin because FPluginReferenceDescriptor doesn't cache any of this info.
+	// GFPs are implemented with a plugin so FPluginReferenceDescriptor doesn't know anything about them.
+	// Need a better way of storing GFP specific plugin data...
+
 	// Read the file to a string
 	FString FileContents;
 	if (!FFileHelper::LoadFileToString(FileContents, *PluginDescriptorFilename))
@@ -1102,7 +1070,7 @@ UGameFeaturePluginStateMachine* UGameFeaturesSubsystem::FindOrCreateGameFeatureP
 
 	UGameFeaturePluginStateMachine* NewStateMachine = NewObject<UGameFeaturePluginStateMachine>(this);
 	GameFeaturePluginStateMachines.Add(PluginURL, NewStateMachine);
-	NewStateMachine->InitStateMachine(PluginURL, FGameFeaturePluginRequestStateMachineDependencies::CreateUObject(this, &ThisClass::HandleRequestPluginDependencyStateMachines));
+	NewStateMachine->InitStateMachine(PluginURL);
 
 	return NewStateMachine;
 }
@@ -1140,7 +1108,7 @@ void UGameFeaturesSubsystem::TerminateGameFeaturePluginComplete(UGameFeaturePlug
 	CompleteDelegate.ExecuteIfBound(Result);
 }
 
-bool UGameFeaturesSubsystem::HandleRequestPluginDependencyStateMachines(const FString& PluginFilename, TArray<UGameFeaturePluginStateMachine*>& OutDependencyMachines)
+bool UGameFeaturesSubsystem::FindOrCreatePluginDependencyStateMachines(const FString& PluginFilename, TArray<UGameFeaturePluginStateMachine*>& OutDependencyMachines)
 {
 	FGameFeaturePluginDetails Details;
 	if (GetGameFeaturePluginDetails(PluginFilename, Details))
@@ -1215,7 +1183,7 @@ TSet<FString> UGameFeaturesSubsystem::GetActivePluginNames() const
 	{
 		UGameFeaturePluginStateMachine* StateMachine = Pair.Value;
 		if (StateMachine->GetCurrentState() == EGameFeaturePluginState::Active &&
-			StateMachine->GetDestinationState() == EGameFeaturePluginState::Active)
+			StateMachine->GetDestination().Contains(EGameFeaturePluginState::Active))
 		{
 			ActivePluginNames.Add(StateMachine->GetPluginName());
 		}
diff --git a/Engine/Plugins/Experimental/GameFeatures/Source/GameFeatures/Public/GameFeaturesSubsystem.h b/Engine/Plugins/Experimental/GameFeatures/Source/GameFeatures/Public/GameFeaturesSubsystem.h
--- a/Engine/Plugins/Experimental/GameFeatures/Source/GameFeatures/Public/GameFeaturesSubsystem.h
+++ b/Engine/Plugins/Experimental/GameFeatures/Source/GameFeatures/Public/GameFeaturesSubsystem.h
@@ -124,7 +124,8 @@ enum class EGameFeatureTargetState : uint8
 	Installed,
 	Registered,
 	Loaded,
-	Active
+	Active,
+	Count	UMETA(Hidden)
 };
 
 struct FGameFeaturePluginDetails
@@ -375,7 +376,8 @@ private:
 	void ChangeGameFeatureTargetStateComplete(UGameFeaturePluginStateMachine* Machine, const UE::GameFeatures::FResult& Result, FGameFeaturePluginUninstallComplete CompleteDelegate);
 
 	/** Handler for when a state machine requests its dependencies. Returns false if the dependencies could not be read */
-	bool HandleRequestPluginDependencyStateMachines(const FString& PluginFilename, TArray<UGameFeaturePluginStateMachine*>& OutDependencyMachines);
+	bool FindOrCreatePluginDependencyStateMachines(const FString& PluginFilename, TArray<UGameFeaturePluginStateMachine*>& OutDependencyMachines);
+	friend struct FGameFeaturePluginState_WaitingForDependencies;
 
 	/** Handle 'ListGameFeaturePlugins' console command */
 	void ListGameFeaturePlugins(const TArray<FString>& Args, UWorld* InWorld, FOutputDevice& Ar);
