Lower severity of a check when unmounting a plugin to instead ensure and set an error state. You can set it back to Registered after this to fix the plugin that failed to unmount

diff --git a/Engine/Plugins/Experimental/GameFeatures/Source/GameFeatures/Private/GameFeaturePluginStateMachine.cpp b/Engine/Plugins/Experimental/GameFeatures/Source/GameFeatures/Private/GameFeaturePluginStateMachine.cpp
--- a/Engine/Plugins/Experimental/GameFeatures/Source/GameFeatures/Private/GameFeaturePluginStateMachine.cpp
+++ b/Engine/Plugins/Experimental/GameFeatures/Source/GameFeatures/Private/GameFeaturePluginStateMachine.cpp
@@ -809,7 +809,13 @@ struct FGameFeaturePluginState_Unmounting : public FGameFeaturePluginState
 		{
 			// The asset registry listens to FPackageName::OnContentPathDismounted() and 
 			// will automatically cleanup the asset registry state we added for this plugin.
-			verify(IPluginManager::Get().UnmountExplicitlyLoadedPlugin(StateProperties.PluginName, nullptr));
+			FText FailureReason;
+			if (!IPluginManager::Get().UnmountExplicitlyLoadedPlugin(StateProperties.PluginName, &FailureReason))
+			{
+				ensureMsgf(false, TEXT("Failed to explicitly unmount the PluginURL(%s) because %s"), *StateProperties.PluginURL, *FailureReason.ToString());
+				Result = MakeError(UE::GameFeatures::StateMachineErrorNamespace + TEXT("BundleManager_Cannot_Explicitly_Unmount"));
+				return;
+			}
 		}
 
 		if (StateProperties.bAddedPluginToManager)
