Change an error message key to be more accurate

diff --git a/Engine/Plugins/Experimental/GameFeatures/Source/GameFeatures/Private/GameFeaturePluginStateMachine.cpp b/Engine/Plugins/Experimental/GameFeatures/Source/GameFeatures/Private/GameFeaturePluginStateMachine.cpp
--- a/Engine/Plugins/Experimental/GameFeatures/Source/GameFeatures/Private/GameFeaturePluginStateMachine.cpp
+++ b/Engine/Plugins/Experimental/GameFeatures/Source/GameFeatures/Private/GameFeaturePluginStateMachine.cpp
@@ -813,7 +813,7 @@ struct FGameFeaturePluginState_Unmounting : public FGameFeaturePluginState
 			if (!IPluginManager::Get().UnmountExplicitlyLoadedPlugin(StateProperties.PluginName, &FailureReason))
 			{
 				ensureMsgf(false, TEXT("Failed to explicitly unmount the PluginURL(%s) because %s"), *StateProperties.PluginURL, *FailureReason.ToString());
-				Result = MakeError(UE::GameFeatures::StateMachineErrorNamespace + TEXT("BundleManager_Cannot_Explicitly_Unmount"));
+				Result = MakeError(UE::GameFeatures::StateMachineErrorNamespace + TEXT("Plugin_Cannot_Explicitly_Unmount"));
 				return;
 			}
 		}
