Cleanup asset manager when unloading game feature plugins.
- Remove game feature data from asset manager when unregistering.
- When an asset is removed from the asset registry, also remove it from the asset manager.
-review-20397299
#p4v-preflight-copy 20385644

diff --git a/Engine/Plugins/Experimental/GameFeatures/Source/GameFeatures/Private/GameFeaturePluginStateMachine.cpp b/Engine/Plugins/Experimental/GameFeatures/Source/GameFeatures/Private/GameFeaturePluginStateMachine.cpp
--- a/Engine/Plugins/Experimental/GameFeatures/Source/GameFeatures/Private/GameFeaturePluginStateMachine.cpp
+++ b/Engine/Plugins/Experimental/GameFeatures/Source/GameFeatures/Private/GameFeaturePluginStateMachine.cpp
@@ -834,6 +834,8 @@ struct FGameFeaturePluginState_Unmounting : public FGameFeaturePluginState
 		{
 			// The asset registry listens to FPackageName::OnContentPathDismounted() and 
 			// will automatically cleanup the asset registry state we added for this plugin.
+			// This will also cause any assets we added to the asset manager to be removed.
+			// Scan paths added to the asset manager should have already been cleaned up.
 			FText FailureReason;
 			if (!IPluginManager::Get().UnmountExplicitlyLoadedPlugin(StateProperties.PluginName, &FailureReason))
 			{
@@ -1296,6 +1298,10 @@ struct FGameFeaturePluginState_Unregistering : public FGameFeaturePluginState
 		if (StateProperties.GameFeatureData)
 		{
 			UGameFeaturesSubsystem::Get().OnGameFeatureUnregistering(StateProperties.GameFeatureData, StateProperties.PluginName, StateProperties.PluginURL);
+
+			UGameFeaturesSubsystem::Get().RemoveGameFeatureFromAssetManager(StateProperties.GameFeatureData, StateProperties.PluginName, StateProperties.AddedPrimaryAssetTypes);
+			StateProperties.AddedPrimaryAssetTypes.Empty();
+
 			UGameFeaturesSubsystem::Get().UnloadGameFeatureData(StateProperties.GameFeatureData);
 		}
 
@@ -1360,6 +1366,9 @@ struct FGameFeaturePluginState_Registering : public FGameFeaturePluginState
 			StateProperties.GameFeatureData->InitializeBasePluginIniFile(StateProperties.PluginInstalledFilename);
 			StateStatus.SetTransition(EGameFeaturePluginState::Registered);
 
+			check(StateProperties.AddedPrimaryAssetTypes.Num() == 0);
+			UGameFeaturesSubsystem::Get().AddGameFeatureToAssetManager(StateProperties.GameFeatureData, StateProperties.PluginName, StateProperties.AddedPrimaryAssetTypes);
+
 			UGameFeaturesSubsystem::Get().OnGameFeatureRegistering(StateProperties.GameFeatureData, StateProperties.PluginName, StateProperties.PluginURL);
 		}
 		else
diff --git a/Engine/Plugins/Experimental/GameFeatures/Source/GameFeatures/Private/GameFeaturePluginStateMachine.h b/Engine/Plugins/Experimental/GameFeatures/Source/GameFeatures/Private/GameFeaturePluginStateMachine.h
--- a/Engine/Plugins/Experimental/GameFeatures/Source/GameFeatures/Private/GameFeaturePluginStateMachine.h
+++ b/Engine/Plugins/Experimental/GameFeatures/Source/GameFeatures/Private/GameFeaturePluginStateMachine.h
@@ -239,7 +239,9 @@
 	/** Meta data parsed from the URL for a specific protocol. */
 	TUnion<UE::GameFeatures::FInstallBundlePluginProtocolMetaData> ProtocolMetadata;
 
-	/** Tracks whether or not this state machine added the pluging to the plugin manager. */
+	TArray<FName> AddedPrimaryAssetTypes;
+
+	/** Tracks whether or not this state machine added the plugin to the plugin manager. */
 	bool bAddedPluginToManager = false;
 
 	/** Whether this state machine should attempt to cancel the current transition */
diff --git a/Engine/Plugins/Experimental/GameFeatures/Source/GameFeatures/Private/GameFeaturesSubsystem.cpp b/Engine/Plugins/Experimental/GameFeatures/Source/GameFeatures/Private/GameFeaturesSubsystem.cpp
--- a/Engine/Plugins/Experimental/GameFeatures/Source/GameFeatures/Private/GameFeaturesSubsystem.cpp
+++ b/Engine/Plugins/Experimental/GameFeatures/Source/GameFeatures/Private/GameFeaturesSubsystem.cpp
@@ -311,7 +311,7 @@ void UGameFeaturesSubsystem::UnloadGameFeatureData(const UGameFeatureData* GameF
 	LocalAssetManager.UnloadPrimaryAsset(GameFeatureToUnload->GetPrimaryAssetId());
 }
 
-void UGameFeaturesSubsystem::AddGameFeatureToAssetManager(const UGameFeatureData* GameFeatureToAdd, const FString& PluginName)
+void UGameFeaturesSubsystem::AddGameFeatureToAssetManager(const UGameFeatureData* GameFeatureToAdd, const FString& PluginName, TArray<FName>& OutNewPrimaryAssetTypes)
 {
 	check(GameFeatureToAdd);
 	FString PluginRootPath = TEXT("/") + PluginName + TEXT("/");
@@ -340,6 +340,8 @@ void UGameFeaturesSubsystem::AddGameFeatureToAssetManager(const UGameFeatureData
 
 		if (!bAlreadyExisted)
 		{
+			OutNewPrimaryAssetTypes.Add(TypeInfo.PrimaryAssetType);
+
 			// If we did not previously scan anything for a primary asset type that is in our config, try to reuse the cook rules from the config instead of the one in the gamefeaturedata, which should not be modifying cook rules
 			const FPrimaryAssetTypeInfo* ConfigTypeInfo = LocalAssetManager.GetSettings().PrimaryAssetTypesToScan.FindByPredicate([&TypeInfo](const FPrimaryAssetTypeInfo& PATI) -> bool { return PATI.PrimaryAssetType == TypeInfo.PrimaryAssetType; });
 			if (ConfigTypeInfo)
@@ -356,9 +358,32 @@ void UGameFeaturesSubsystem::AddGameFeatureToAssetManager(const UGameFeatureData
 	LocalAssetManager.PopBulkScanning();
 }
 
-void UGameFeaturesSubsystem::RemoveGameFeatureFromAssetManager(const UGameFeatureData* GameFeatureToRemove)
+void UGameFeaturesSubsystem::RemoveGameFeatureFromAssetManager(const UGameFeatureData* GameFeatureToRemove, const FString& PluginName, const TArray<FName>& AddedPrimaryAssetTypes)
 {
-	/** NOT IMPLEMENTED - STUB */
+	check(GameFeatureToRemove);
+	UAssetManager& LocalAssetManager = UAssetManager::Get();
+
+	for (FPrimaryAssetTypeInfo TypeInfo : GameFeatureToRemove->GetPrimaryAssetTypesToScan())
+	{
+		if (AddedPrimaryAssetTypes.Contains(TypeInfo.PrimaryAssetType))
+		{
+			LocalAssetManager.RemovePrimaryAssetType(TypeInfo.PrimaryAssetType);
+			continue;
+		}
+
+		for (FDirectoryPath& Path : TypeInfo.Directories)
+		{
+			Path.Path = TEXT("/") + PluginName + TEXT("/") + Path.Path;
+		}
+
+		// This function also fills out runtime data on the copy
+		if (!LocalAssetManager.ShouldScanPrimaryAssetType(TypeInfo))
+		{
+			continue;
+		}
+
+		LocalAssetManager.RemoveScanPathsForPrimaryAssets(TypeInfo.PrimaryAssetType, TypeInfo.AssetScanPaths, TypeInfo.AssetBaseClassLoaded, TypeInfo.bHasBlueprintClasses, TypeInfo.bIsEditorOnly);
+	}
 }
 
 void UGameFeaturesSubsystem::AddObserver(UObject* Observer)
@@ -464,7 +489,6 @@ void UGameFeaturesSubsystem::OnGameFeatureStatusKnown(const FString& PluginName,
 void UGameFeaturesSubsystem::OnGameFeatureRegistering(const UGameFeatureData* GameFeatureData, const FString& PluginName, const FString& PluginURL)
 {
 	check(GameFeatureData);
-	AddGameFeatureToAssetManager(GameFeatureData, PluginName);
 
 	for (UObject* Observer : Observers)
 	{
@@ -549,8 +573,6 @@ void UGameFeaturesSubsystem::OnGameFeatureDeactivating(const UGameFeatureData* G
 			Action->OnGameFeatureDeactivating(Context);
 		}
 	}
-
-	RemoveGameFeatureFromAssetManager(GameFeatureData);
 }
 
 const UGameFeatureData* UGameFeaturesSubsystem::GetDataForStateMachine(UGameFeaturePluginStateMachine* GFSM) const
diff --git a/Engine/Plugins/Experimental/GameFeatures/Source/GameFeatures/Public/GameFeaturesSubsystem.h b/Engine/Plugins/Experimental/GameFeatures/Source/GameFeatures/Public/GameFeaturesSubsystem.h
--- a/Engine/Plugins/Experimental/GameFeatures/Source/GameFeatures/Public/GameFeaturesSubsystem.h
+++ b/Engine/Plugins/Experimental/GameFeatures/Source/GameFeatures/Public/GameFeaturesSubsystem.h
@@ -351,9 +351,9 @@
 	void OnAssetManagerCreated();
 
 	/** Scans for assets specified in the game feature data */
-	static void AddGameFeatureToAssetManager(const UGameFeatureData* GameFeatureToAdd, const FString& PluginName);
+	static void AddGameFeatureToAssetManager(const UGameFeatureData* GameFeatureToAdd, const FString& PluginName, TArray<FName>& OutNewPrimaryAssetTypes);
 
-	static void RemoveGameFeatureFromAssetManager(const UGameFeatureData* GameFeatureToRemove);
+	static void RemoveGameFeatureFromAssetManager(const UGameFeatureData* GameFeatureToRemove, const FString& PluginName, const TArray<FName>& AddedPrimaryAssetTypes);
 
 private:
 	const UGameFeatureData* GetDataForStateMachine(UGameFeaturePluginStateMachine* GFSM) const;
diff --git a/Engine/Source/Runtime/Engine/Classes/Engine/AssetManager.h b/Engine/Source/Runtime/Engine/Classes/Engine/AssetManager.h
--- a/Engine/Source/Runtime/Engine/Classes/Engine/AssetManager.h
+++ b/Engine/Source/Runtime/Engine/Classes/Engine/AssetManager.h
@@ -97,11 +97,15 @@ public:
 
 	/** Single path wrapper */
 	virtual int32 ScanPathForPrimaryAssets(FPrimaryAssetType PrimaryAssetType, const FString& Path, UClass* BaseClass, bool bHasBlueprintClasses, bool bIsEditorOnly = false, bool bForceSynchronousScan = true);
-	
+
 	/** Call before many calls to ScanPaths to improve load performance. Match each call with PopBulkScanning(). */
 	void PushBulkScanning();
 	void PopBulkScanning();
 
+	virtual void RemoveScanPathsForPrimaryAssets(FPrimaryAssetType PrimaryAssetType, const TArray<FString>& Paths, UClass* BaseClass, bool bHasBlueprintClasses, bool bIsEditorOnly = false);
+
+	virtual void RemovePrimaryAssetType(FPrimaryAssetType PrimaryAssetType);
+
 protected:
 	/** Should only be called from PushBulkScanning() and override */
 	virtual void StartBulkScanning();
@@ -669,6 +673,9 @@ protected:
 	/** When asset is renamed */
 	virtual void OnAssetRenamed(const FAssetData& NewData, const FString& OldPath);
 
+	/** When an asset is removed */
+	virtual void OnAssetRemoved(const FAssetData& Data);
+
 	/** Try to remove an old asset identifier when it has been deleted/renamed */
 	virtual void RemovePrimaryAssetId(const FPrimaryAssetId& PrimaryAssetId);
 
diff --git a/Engine/Source/Runtime/Engine/Classes/Engine/AssetManagerTypes.h b/Engine/Source/Runtime/Engine/Classes/Engine/AssetManagerTypes.h
--- a/Engine/Source/Runtime/Engine/Classes/Engine/AssetManagerTypes.h
+++ b/Engine/Source/Runtime/Engine/Classes/Engine/AssetManagerTypes.h
@@ -130,9 +130,10 @@ public:
 	bool bIsEditorOnly;
 
 #if WITH_EDITOR
+	const TArray<FDirectoryPath>& GetDirectories() const { return Directories; }
 	TArray<FDirectoryPath>& GetDirectories() { return Directories; }
 
-	TArray<FSoftObjectPath>& GetSpecificAssets() { return SpecificAssets; }
+	const TArray<FSoftObjectPath>& GetSpecificAssets() const { return SpecificAssets; }
 #endif //if WITH_EDITOR
 
 private:
diff --git a/Engine/Source/Runtime/Engine/Private/AssetManager.cpp b/Engine/Source/Runtime/Engine/Private/AssetManager.cpp
--- a/Engine/Source/Runtime/Engine/Private/AssetManager.cpp
+++ b/Engine/Source/Runtime/Engine/Private/AssetManager.cpp
@@ -294,6 +294,7 @@ void UAssetManager::PostInitProperties()
 			AssetRegistry.OnInMemoryAssetCreated().AddUObject(this, &UAssetManager::OnInMemoryAssetCreated);
 			AssetRegistry.OnInMemoryAssetDeleted().AddUObject(this, &UAssetManager::OnInMemoryAssetDeleted);
 			AssetRegistry.OnAssetRenamed().AddUObject(this, &UAssetManager::OnAssetRenamed);
+			AssetRegistry.OnAssetRemoved().AddUObject(this, &UAssetManager::OnAssetRemoved);
 		}
 
 		FEditorDelegates::PreBeginPIE.AddUObject(this, &UAssetManager::PreBeginPIE);
@@ -810,13 +811,14 @@ int32 UAssetManager::ScanPathsForPrimaryAssets(FPrimaryAssetType PrimaryAssetTyp
 {
 	LLM_SCOPE_BYTAG(AssetManager);
 	TRACE_CPUPROFILER_EVENT_SCOPE(UAssetManager::ScanPathsForPrimaryAssets)
-	TSharedRef<FPrimaryAssetTypeData>* FoundType = AssetTypeMap.Find(PrimaryAssetType);
 
 	if (bIsEditorOnly && !GIsEditor)
 	{
 		return 0;
 	}
 
+	TSharedRef<FPrimaryAssetTypeData>* FoundType = AssetTypeMap.Find(PrimaryAssetType);
+
 	check(BaseClass);
 
 	if (!FoundType)
@@ -938,6 +940,55 @@ void UAssetManager::PopBulkScanning()
 	}
 }
 
+void UAssetManager::RemoveScanPathsForPrimaryAssets(FPrimaryAssetType PrimaryAssetType, const TArray<FString>& Paths, UClass* BaseClass, bool bHasBlueprintClasses, bool bIsEditorOnly /*= false*/)
+{
+	if (bIsEditorOnly && !GIsEditor)
+	{
+		return;
+	}
+
+	TSharedRef<FPrimaryAssetTypeData>* FoundType = AssetTypeMap.Find(PrimaryAssetType);
+
+	check(BaseClass);
+
+	if (!FoundType)
+	{
+		return;
+	}
+
+	FPrimaryAssetTypeData& TypeData = FoundType->Get();
+
+	// Make sure types match
+	if (!ensureMsgf(TypeData.Info.AssetBaseClassLoaded == BaseClass && TypeData.Info.bHasBlueprintClasses == bHasBlueprintClasses && TypeData.Info.bIsEditorOnly == bIsEditorOnly, TEXT("UAssetManager::RemoveScanPathsForPrimaryAssets TypeData parameters did not match for type '%s'"), *TypeData.Info.PrimaryAssetType.ToString()))
+	{
+		return;
+	}
+
+	TArray<FString> RemovedPaths;
+	RemovedPaths.Reserve(Paths.Num());
+	for (const FString& Path : Paths)
+	{
+		if (TypeData.Info.AssetScanPaths.Remove(Path))
+		{
+			RemovedPaths.Add(Path);
+		}
+
+		TypeData.DeferredAssetScanPaths.Remove(Path);
+	}
+
+	// Expand paths so we can record them for later
+	ExpandVirtualPaths(RemovedPaths);
+	for (const FString& Path : RemovedPaths)
+	{
+		TypeData.RealAssetScanPaths.Remove(Path);
+	}
+}
+
+void UAssetManager::RemovePrimaryAssetType(FPrimaryAssetType PrimaryAssetType)
+{
+	AssetTypeMap.Remove(PrimaryAssetType);
+}
+
 void UAssetManager::StartBulkScanning()
 {
 	check(IsBulkScanning());
@@ -4403,6 +4454,23 @@ void UAssetManager::OnAssetRenamed(const FAssetData& NewData, const FString& Old
 	OnInMemoryAssetCreated(NewObject);
 }
 
+void UAssetManager::OnAssetRemoved(const FAssetData& Data)
+{
+	// This could be much more efficient if UAssetManager broadcast one large event instead of all these tiny updates, see UAssetRegistryImpl::Broadcast
+
+	FPrimaryAssetId PrimaryAssetId = GetPrimaryAssetIdForPath(Data.ObjectPath);
+
+	// This may be a blueprint, try with _C
+	if (!PrimaryAssetId.IsValid())
+	{
+		PrimaryAssetId = GetPrimaryAssetIdForPath(FName(FNameBuilder(Data.ObjectPath).Append(TEXT("_C")).ToView(), FNAME_Find));
+	}
+
+	CachedAssetBundles.Remove(PrimaryAssetId);
+
+	RemovePrimaryAssetId(PrimaryAssetId);
+}
+
 void UAssetManager::RemovePrimaryAssetId(const FPrimaryAssetId& PrimaryAssetId)
 {
 	if (PrimaryAssetId.IsValid() && GetNameData(PrimaryAssetId))
