Fix gamefeature subsystem observer array being modified while iterated for the purposes of dispatching callbacks.

diff --git a/Engine/Plugins/Experimental/GameFeatures/Source/GameFeatures/Private/GameFeaturePluginStateMachine.cpp b/Engine/Plugins/Experimental/GameFeatures/Source/GameFeatures/Private/GameFeaturePluginStateMachine.cpp
--- a/Engine/Plugins/Experimental/GameFeatures/Source/GameFeatures/Private/GameFeaturePluginStateMachine.cpp
+++ b/Engine/Plugins/Experimental/GameFeatures/Source/GameFeatures/Private/GameFeaturePluginStateMachine.cpp
@@ -1833,12 +1833,12 @@ void UGameFeaturePluginStateMachine::RemovePendingCancelCallback(void* DelegateO
 	StateProperties.OnTransitionCanceled.RemoveAll(DelegateObject);
 }
 
-FString UGameFeaturePluginStateMachine::GetGameFeatureName() const
+const FString& UGameFeaturePluginStateMachine::GetGameFeatureName() const
 {
 	FString PluginFilename;
-	if (GetPluginFilename(PluginFilename))
+	if (!StateProperties.PluginName.IsEmpty())
 	{
-		return FPaths::GetBaseFilename(PluginFilename);
+		return StateProperties.PluginName;
 	}
 	else
 	{
@@ -1846,12 +1846,12 @@ FString UGameFeaturePluginStateMachine::GetGameFeatureName() const
 	}
 }
 
-FString UGameFeaturePluginStateMachine::GetPluginURL() const
+const FString& UGameFeaturePluginStateMachine::GetPluginURL() const
 {
 	return StateProperties.PluginURL;
 }
 
-FString UGameFeaturePluginStateMachine::GetPluginName() const
+const FString& UGameFeaturePluginStateMachine::GetPluginName() const
 {
 	return StateProperties.PluginName;
 }
diff --git a/Engine/Plugins/Experimental/GameFeatures/Source/GameFeatures/Private/GameFeaturePluginStateMachine.h b/Engine/Plugins/Experimental/GameFeatures/Source/GameFeatures/Private/GameFeaturePluginStateMachine.h
--- a/Engine/Plugins/Experimental/GameFeatures/Source/GameFeatures/Private/GameFeaturePluginStateMachine.h
+++ b/Engine/Plugins/Experimental/GameFeatures/Source/GameFeatures/Private/GameFeaturePluginStateMachine.h
@@ -422,13 +422,13 @@ public:
 	void RemovePendingCancelCallback(void* DelegateObject);
 
 	/** Returns the name of the game feature. Before StatusKnown, this returns the URL. */
-	FString GetGameFeatureName() const;
+	const FString& GetGameFeatureName() const;
 
 	/** Returns the URL */
-	FString GetPluginURL() const;
+	const FString& GetPluginURL() const;
 
 	/** Returns the plugin name if known (plugin must have been registered to know the name). */
-	FString GetPluginName() const;
+	const FString& GetPluginName() const;
 
 	/** Returns the uplugin filename of the game feature. Before StatusKnown, this returns false. */
 	bool GetPluginFilename(FString& OutPluginFilename) const;
diff --git a/Engine/Plugins/Experimental/GameFeatures/Source/GameFeatures/Private/GameFeaturesSubsystem.cpp b/Engine/Plugins/Experimental/GameFeatures/Source/GameFeatures/Private/GameFeaturesSubsystem.cpp
--- a/Engine/Plugins/Experimental/GameFeatures/Source/GameFeatures/Private/GameFeaturesSubsystem.cpp
+++ b/Engine/Plugins/Experimental/GameFeatures/Source/GameFeatures/Private/GameFeaturesSubsystem.cpp
@@ -392,14 +392,14 @@ void UGameFeaturesSubsystem::AddObserver(UObject* Observer)
 	check(Observer);
 	if (ensureAlwaysMsgf(Cast<IGameFeatureStateChangeObserver>(Observer) != nullptr, TEXT("Observers must implement the IGameFeatureStateChangeObserver interface.")))
 	{
-		Observers.Add(Observer);
+		Observers.AddUnique(Observer);
 	}
 }
 
 void UGameFeaturesSubsystem::RemoveObserver(UObject* Observer)
 {
 	check(Observer);
-	Observers.Remove(Observer);
+	Observers.RemoveSingleSwap(Observer);
 }
 
 FString UGameFeaturesSubsystem::GetPluginURL_FileProtocol(const FString& PluginDescriptorPath)
@@ -457,10 +457,7 @@ EGameFeaturePluginProtocol UGameFeaturesSubsystem::GetPluginURLProtocol(FStringV
 
 void UGameFeaturesSubsystem::OnGameFeatureTerminating(const FString& PluginName, const FString& PluginURL)
 {
-	for (UObject* Observer : Observers)
-	{
-		CastChecked<IGameFeatureStateChangeObserver>(Observer)->OnGameFeatureTerminating(PluginURL);
-	}
+	CallbackObservers(EObserverCallback::Terminating, PluginURL, &PluginName);
 
 	if (!PluginName.IsEmpty())
 	{
@@ -471,10 +468,7 @@ void UGameFeaturesSubsystem::OnGameFeatureTerminating(const FString& PluginName,
 
 void UGameFeaturesSubsystem::OnGameFeatureCheckingStatus(const FString& PluginURL)
 {
-	for (UObject* Observer : Observers)
-	{
-		CastChecked<IGameFeatureStateChangeObserver>(Observer)->OnGameFeatureCheckingStatus(PluginURL);
-	}
+	CallbackObservers(EObserverCallback::CheckingStatus, PluginURL);
 }
 
 void UGameFeaturesSubsystem::OnGameFeatureStatusKnown(const FString& PluginName, const FString& PluginURL)
@@ -488,12 +482,7 @@ void UGameFeaturesSubsystem::OnGameFeatureStatusKnown(const FString& PluginName,
 
 void UGameFeaturesSubsystem::OnGameFeatureRegistering(const UGameFeatureData* GameFeatureData, const FString& PluginName, const FString& PluginURL)
 {
-	check(GameFeatureData);
-
-	for (UObject* Observer : Observers)
-	{
-		CastChecked<IGameFeatureStateChangeObserver>(Observer)->OnGameFeatureRegistering(GameFeatureData, PluginName, PluginURL);
-	}
+	CallbackObservers(EObserverCallback::Registering, PluginURL, &PluginName, GameFeatureData);
 
 	for (UGameFeatureAction* Action : GameFeatureData->GetActions())
 	{
@@ -506,12 +495,7 @@ void UGameFeaturesSubsystem::OnGameFeatureRegistering(const UGameFeatureData* Ga
 
 void UGameFeaturesSubsystem::OnGameFeatureUnregistering(const UGameFeatureData* GameFeatureData, const FString& PluginName, const FString& PluginURL)
 {
-	check(GameFeatureData);
-
-	for (UObject* Observer : Observers)
-	{
-		CastChecked<IGameFeatureStateChangeObserver>(Observer)->OnGameFeatureUnregistering(GameFeatureData, PluginName, PluginURL);
-	}
+	CallbackObservers(EObserverCallback::Unregistering, PluginURL, &PluginName, GameFeatureData);
 
 	for (UGameFeatureAction* Action : GameFeatureData->GetActions())
 	{
@@ -524,11 +508,7 @@ void UGameFeaturesSubsystem::OnGameFeatureUnregistering(const UGameFeatureData*
 
 void UGameFeaturesSubsystem::OnGameFeatureLoading(const UGameFeatureData* GameFeatureData, const FString& PluginURL)
 {
-	check(GameFeatureData);
-	for (UObject* Observer : Observers)
-	{
-		CastChecked<IGameFeatureStateChangeObserver>(Observer)->OnGameFeatureLoading(GameFeatureData, PluginURL);
-	}
+	CallbackObservers(EObserverCallback::Loading, PluginURL, nullptr, GameFeatureData);
 
 	for (UGameFeatureAction* Action : GameFeatureData->GetActions())
 	{
@@ -541,13 +521,8 @@ void UGameFeaturesSubsystem::OnGameFeatureLoading(const UGameFeatureData* GameFe
 
 void UGameFeaturesSubsystem::OnGameFeatureActivating(const UGameFeatureData* GameFeatureData, const FString& PluginName, FGameFeatureActivatingContext& Context, const FString& PluginURL)
 {
-	check(GameFeatureData);
+	CallbackObservers(EObserverCallback::Activating, PluginURL, &PluginName, GameFeatureData);
 
-	for (UObject* Observer : Observers)
-	{
-		CastChecked<IGameFeatureStateChangeObserver>(Observer)->OnGameFeatureActivating(GameFeatureData, PluginURL);
-	}
-	
 	for (UGameFeatureAction* Action : GameFeatureData->GetActions())
 	{
 		if (Action != nullptr)
@@ -559,12 +534,7 @@ void UGameFeaturesSubsystem::OnGameFeatureActivating(const UGameFeatureData* Gam
 
 void UGameFeaturesSubsystem::OnGameFeatureDeactivating(const UGameFeatureData* GameFeatureData, const FString& PluginName, FGameFeatureDeactivatingContext& Context, const FString& PluginURL)
 {
-	check(GameFeatureData);
-
-	for (UObject* Observer : Observers)
-	{
-		CastChecked<IGameFeatureStateChangeObserver>(Observer)->OnGameFeatureDeactivating(GameFeatureData, Context, PluginURL);
-	}
+	CallbackObservers(EObserverCallback::Deactivating, PluginURL, &PluginName, GameFeatureData, &Context);
 
 	for (UGameFeatureAction* Action : GameFeatureData->GetActions())
 	{
@@ -655,10 +625,7 @@ void UGameFeaturesSubsystem::LoadGameFeaturePlugin(const FString& PluginURL, con
 		// on whichever Role runs second between client and server.
 
 		// Refire the observer for Activated and do nothing else.
-		for (UObject* Observer : Observers)
-		{
-			CastChecked<IGameFeatureStateChangeObserver>(Observer)->OnGameFeatureActivating(StateMachine->GetGameFeatureDataForActivePlugin(), PluginURL);
-		}
+		CallbackObservers(EObserverCallback::Activating, PluginURL, &StateMachine->GetPluginName(), StateMachine->GetGameFeatureDataForActivePlugin());
 	}
 
 	ChangeGameFeatureDestination(StateMachine, FGameFeaturePluginStateRange(EGameFeaturePluginState::Loaded, EGameFeaturePluginState::Active), CompleteDelegate);
@@ -722,10 +689,7 @@ void UGameFeaturesSubsystem::ChangeGameFeatureTargetState(const FString& PluginU
 		// on whichever Role runs second between client and server.
 
 		// Refire the observer for Activated and do nothing else.
-		for (UObject* Observer : Observers)
-		{
-			CastChecked<IGameFeatureStateChangeObserver>(Observer)->OnGameFeatureActivating(StateMachine->GetGameFeatureDataForActivePlugin(), PluginURL);
-		}
+		CallbackObservers(EObserverCallback::Activating, PluginURL, &StateMachine->GetPluginName(), StateMachine->GetGameFeatureDataForActivePlugin());
 	}
 	
 	ChangeGameFeatureDestination(StateMachine, FGameFeaturePluginStateRange(TargetPluginState), CompleteDelegate);
@@ -928,11 +892,8 @@ bool UGameFeaturesSubsystem::GetPluginURLForBuiltInPluginByName(const FString& P
 FString UGameFeaturesSubsystem::GetPluginFilenameFromPluginURL(const FString& PluginURL) const
 {
 	FString PluginFilename;
-	if (const UGameFeaturePluginStateMachine* GFSM = FindGameFeaturePluginStateMachine(PluginURL))
-	{
-		GFSM->GetPluginFilename(PluginFilename);
-	}
-	else
+	const UGameFeaturePluginStateMachine* GFSM = FindGameFeaturePluginStateMachine(PluginURL);
+	if (GFSM == nullptr || !GFSM->GetPluginFilename(PluginFilename))
 	{
 		UE_LOG(LogGameFeatures, Error, TEXT("UGameFeaturesSubsystem could not get the plugin path from the plugin URL. URL:%s "), *PluginURL);
 	}
@@ -1275,6 +1236,78 @@ void UGameFeaturesSubsystem::ListGameFeaturePlugins(const TArray<FString>& Args,
 	Ar.Logf(TEXT("Total Game Feature Plugins: %d"), PluginCount);
 }
 
+void UGameFeaturesSubsystem::CallbackObservers(EObserverCallback CallbackType, const FString& PluginURL, 
+	const FString* PluginName /*= nullptr*/, 
+	const UGameFeatureData* GameFeatureData /*= nullptr*/, 
+	FGameFeatureDeactivatingContext* DeactivatingContext /*= nullptr*/)
+{
+	// Protect against modifying the observer list during iteration
+	TArray<UObject*> LocalObservers(Observers);
+
+	switch (CallbackType)
+	{
+	case EObserverCallback::CheckingStatus:
+		for (UObject* Observer : LocalObservers)
+		{
+			CastChecked<IGameFeatureStateChangeObserver>(Observer)->OnGameFeatureCheckingStatus(PluginURL);
+		}
+		break;
+
+	case EObserverCallback::Terminating:
+		for (UObject* Observer : LocalObservers)
+		{
+			CastChecked<IGameFeatureStateChangeObserver>(Observer)->OnGameFeatureTerminating(PluginURL);
+		}
+		break;
+
+	case EObserverCallback::Registering:
+		check(PluginName);
+		check(GameFeatureData);
+		for (UObject* Observer : LocalObservers)
+		{
+			CastChecked<IGameFeatureStateChangeObserver>(Observer)->OnGameFeatureRegistering(GameFeatureData, *PluginName, PluginURL);
+		}
+		break;
+
+	case EObserverCallback::Unregistering:
+		check(PluginName);
+		check(GameFeatureData);
+		for (UObject* Observer : LocalObservers)
+		{
+			CastChecked<IGameFeatureStateChangeObserver>(Observer)->OnGameFeatureUnregistering(GameFeatureData, *PluginName, PluginURL);
+		}
+		break;
+
+	case EObserverCallback::Loading:
+		check(GameFeatureData);
+		for (UObject* Observer : LocalObservers)
+		{
+			CastChecked<IGameFeatureStateChangeObserver>(Observer)->OnGameFeatureLoading(GameFeatureData, PluginURL);
+		}
+		break;
+
+	case EObserverCallback::Activating:
+		check(GameFeatureData);
+		for (UObject* Observer : LocalObservers)
+		{
+			CastChecked<IGameFeatureStateChangeObserver>(Observer)->OnGameFeatureActivating(GameFeatureData, PluginURL);
+		}
+		break;
+
+	case EObserverCallback::Deactivating:
+		check(GameFeatureData);
+		check(DeactivatingContext);
+		for (UObject* Observer : LocalObservers)
+		{
+			CastChecked<IGameFeatureStateChangeObserver>(Observer)->OnGameFeatureDeactivating(GameFeatureData, *DeactivatingContext, PluginURL);
+		}
+		break;
+
+	default:
+		UE_LOG(LogGameFeatures, Fatal, TEXT("Unkown EObserverCallback!"));
+	}
+}
+
 TSet<FString> UGameFeaturesSubsystem::GetActivePluginNames() const
 {
 	TSet<FString> ActivePluginNames;
diff --git a/Engine/Plugins/Experimental/GameFeatures/Source/GameFeatures/Public/GameFeaturesSubsystem.h b/Engine/Plugins/Experimental/GameFeatures/Source/GameFeatures/Public/GameFeaturesSubsystem.h
--- a/Engine/Plugins/Experimental/GameFeatures/Source/GameFeatures/Public/GameFeaturesSubsystem.h
+++ b/Engine/Plugins/Experimental/GameFeatures/Source/GameFeatures/Public/GameFeaturesSubsystem.h
@@ -391,6 +391,22 @@ private:
 	/** Handle 'ListGameFeaturePlugins' console command */
 	void ListGameFeaturePlugins(const TArray<FString>& Args, UWorld* InWorld, FOutputDevice& Ar);
 
+	enum class EObserverCallback
+	{
+		CheckingStatus,
+		Terminating,
+		Registering,
+		Unregistering,
+		Loading,
+		Activating,
+		Deactivating
+	};
+
+	void CallbackObservers(EObserverCallback CallbackType, const FString& PluginURL, 
+		const FString* PluginName = nullptr, 
+		const UGameFeatureData* GameFeatureData = nullptr, 
+		FGameFeatureDeactivatingContext* DeactivatingContext = nullptr);
+
 private:
 	/** The list of all game feature plugin state machine objects */
 	UPROPERTY(Transient)
