Fix assumption that built-in game feature plugins would not be reloaded at runtime. This fixes auto test errors that were causes by the dedicated server failing to reload plugins after resetting. Any unloading is now canceled if still in flight when reloading happens.

diff --git a/Engine/Plugins/Experimental/GameFeatures/Source/GameFeatures/Private/GameFeaturesSubsystem.cpp b/Engine/Plugins/Experimental/GameFeatures/Source/GameFeatures/Private/GameFeaturesSubsystem.cpp
--- a/Engine/Plugins/Experimental/GameFeatures/Source/GameFeatures/Private/GameFeaturesSubsystem.cpp
+++ b/Engine/Plugins/Experimental/GameFeatures/Source/GameFeatures/Private/GameFeaturesSubsystem.cpp
@@ -852,8 +852,9 @@ void UGameFeaturesSubsystem::LoadBuiltInGameFeaturePlugin(const TSharedRef<IPlug
 					const EGameFeaturePluginState DestinationState = ConvertInitialFeatureStateToTargetState(InitialAutoState);
 
 					// If we're already at the destination or beyond, don't transition back
-					verify(StateMachine->SetDestination(FGameFeaturePluginStateRange(DestinationState, EGameFeaturePluginState::Active), 
-						FGameFeatureStateTransitionComplete::CreateUObject(this, &ThisClass::LoadBuiltInGameFeaturePluginComplete)));
+					FGameFeaturePluginStateRange Destination(DestinationState, EGameFeaturePluginState::Active);
+					ChangeGameFeatureDestination(StateMachine, Destination, 
+						FGameFeaturePluginChangeStateComplete::CreateUObject(this, &ThisClass::LoadBuiltInGameFeaturePluginComplete, StateMachine, Destination));
 				}
 			}
 		}
@@ -1086,19 +1087,32 @@ UGameFeaturePluginStateMachine* UGameFeaturesSubsystem::FindOrCreateGameFeatureP
 	return NewStateMachine;
 }
 
-void UGameFeaturesSubsystem::LoadBuiltInGameFeaturePluginComplete(UGameFeaturePluginStateMachine* Machine, const UE::GameFeatures::FResult& Result)
+void UGameFeaturesSubsystem::LoadBuiltInGameFeaturePluginComplete(const UE::GameFeatures::FResult& Result, UGameFeaturePluginStateMachine* Machine, FGameFeaturePluginStateRange RequestedDestination)
 {
 	check(Machine);
 	if (Result.HasValue())
 	{
-		UE_LOG(LogGameFeatures, Display, TEXT("Game feature '%s' loaded successfully. Ending state: %s"), *Machine->GetGameFeatureName(), *UE::GameFeatures::ToString(Machine->GetCurrentState()));
+		//@note It's possible for the machine to still be tranitioning at this point as long as it's withing the requested destination range
+		UE_LOG(LogGameFeatures, Display, TEXT("Game feature '%s' loaded successfully. Ending state: %s, [%s, %s]"), 
+			*Machine->GetGameFeatureName(), 
+			*UE::GameFeatures::ToString(Machine->GetCurrentState()),
+			*UE::GameFeatures::ToString(Machine->GetDestination().MinState),
+			*UE::GameFeatures::ToString(Machine->GetDestination().MaxState));
+
+		checkf(RequestedDestination.Contains(Machine->GetCurrentState()), TEXT("Game feature '%s': Ending state %s is not in expected range [%s, %s]"), 
+			*Machine->GetGameFeatureName(), 
+			*UE::GameFeatures::ToString(Machine->GetCurrentState()), 
+			*UE::GameFeatures::ToString(RequestedDestination.MinState), 
+			*UE::GameFeatures::ToString(RequestedDestination.MaxState));
 	}
 	else
 	{
 		const FString ErrorMessage = UE::GameFeatures::ToString(Result);
-		UE_LOG(LogGameFeatures, Error, TEXT("Game feature '%s' load failed. Ending state: %s. Result: %s"),
+		UE_LOG(LogGameFeatures, Error, TEXT("Game feature '%s' load failed. Ending state: %s, [%s, %s]. Result: %s"),
 			*Machine->GetGameFeatureName(),
 			*UE::GameFeatures::ToString(Machine->GetCurrentState()),
+			*UE::GameFeatures::ToString(Machine->GetDestination().MinState),
+			*UE::GameFeatures::ToString(Machine->GetDestination().MaxState),
 			*ErrorMessage);
 	}
 }
diff --git a/Engine/Plugins/Experimental/GameFeatures/Source/GameFeatures/Public/GameFeaturesSubsystem.h b/Engine/Plugins/Experimental/GameFeatures/Source/GameFeatures/Public/GameFeaturesSubsystem.h
--- a/Engine/Plugins/Experimental/GameFeatures/Source/GameFeatures/Public/GameFeaturesSubsystem.h
+++ b/Engine/Plugins/Experimental/GameFeatures/Source/GameFeatures/Public/GameFeaturesSubsystem.h
@@ -372,7 +372,7 @@ private:
 	UGameFeaturePluginStateMachine* FindOrCreateGameFeaturePluginStateMachine(const FString& PluginURL);
 
 	/** Notification that a game feature has finished loading, and whether it was successful */
-	void LoadBuiltInGameFeaturePluginComplete(UGameFeaturePluginStateMachine* Machine, const UE::GameFeatures::FResult& Result);
+	void LoadBuiltInGameFeaturePluginComplete(const UE::GameFeatures::FResult& Result, UGameFeaturePluginStateMachine* Machine, FGameFeaturePluginStateRange RequestedDestination);
 
 	/** Sets a new destination state. Will attempt to cancel the current transition if the new destination is incompatible with the current destination */
 	void ChangeGameFeatureDestination(UGameFeaturePluginStateMachine* Machine, const FGameFeaturePluginStateRange& StateRange, FGameFeaturePluginChangeStateComplete CompleteDelegate);
