Turn on GFP leak detection by CVar by default.
GFP leak detection now runs on non-shipping builds, servers, and editors.
Add an ensure.

diff --git a/Engine/Plugins/Experimental/GameFeatures/Source/GameFeatures/Private/GameFeaturePluginStateMachine.cpp b/Engine/Plugins/Experimental/GameFeatures/Source/GameFeatures/Private/GameFeaturePluginStateMachine.cpp
--- a/Engine/Plugins/Experimental/GameFeatures/Source/GameFeatures/Private/GameFeaturePluginStateMachine.cpp
+++ b/Engine/Plugins/Experimental/GameFeatures/Source/GameFeatures/Private/GameFeaturePluginStateMachine.cpp
@@ -40,7 +40,7 @@ namespace UE::GameFeatures
 		TEXT("Should the newly mounted files be logged."));
 
 	static TAutoConsoleVariable<bool> CVarVerifyPluginUnload(TEXT("GameFeaturePlugin.VerifyUnload"), 
-		false,
+		true,
 		TEXT("Verify plugin assets are no longer in memory when unloading."),
 		ECVF_Default);
 
@@ -65,7 +65,7 @@ namespace UE::GameFeatures
 	// Verify that all assets from this plugin have been unloaded and GC'd	
 	void VerifyAssetsUnloaded(const FString& PluginName, bool bIgnoreGameFeatureData)
 	{
-#if !UE_BUILD_SHIPPING
+#if (!UE_BUILD_SHIPPING || UE_SERVER || WITH_EDITOR)
 		if (!UE::GameFeatures::CVarVerifyPluginUnload.GetValueOnGameThread())
 		{
 			return;
@@ -85,6 +85,8 @@ namespace UE::GameFeatures
 				//EInternalObjectFlags InternalFlags = AssetObj->GetInternalFlags();
 
 				FReferenceChainSearch(AssetObj, EReferenceChainSearchMode::Shortest | EReferenceChainSearchMode::PrintResults);
+
+				ensureAlwaysMsgf(false, TEXT("GFP %s failed to unload asset %s!"), *PluginName, *AssetData.GetFullName());
 			}
 		};
 
