Fix missing call to PopBulkScanning().  This was causing asset registry bulk scanning to never end.

diff --git a/Engine/Plugins/Experimental/GameFeatures/Source/GameFeatures/Private/GameFeaturesSubsystem.cpp b/Engine/Plugins/Experimental/GameFeatures/Source/GameFeatures/Private/GameFeaturesSubsystem.cpp
--- a/Engine/Plugins/Experimental/GameFeatures/Source/GameFeatures/Private/GameFeaturesSubsystem.cpp
+++ b/Engine/Plugins/Experimental/GameFeatures/Source/GameFeatures/Private/GameFeaturesSubsystem.cpp
@@ -313,6 +313,9 @@ void UGameFeaturesSubsystem::UnloadGameFeatureData(const UGameFeatureData* GameF
 
 void UGameFeaturesSubsystem::AddGameFeatureToAssetManager(const UGameFeatureData* GameFeatureToAdd, const FString& PluginName, TArray<FName>& OutNewPrimaryAssetTypes)
 {
+	// JMarcus TODO: FPackageName::RegisterMountPoint
+	// Also need to test cooked project loading to see if asset reg can find primary assets
+
 	check(GameFeatureToAdd);
 	FString PluginRootPath = TEXT("/") + PluginName + TEXT("/");
 	UAssetManager& LocalAssetManager = UAssetManager::Get();
@@ -859,6 +862,8 @@ void UGameFeaturesSubsystem::LoadBuiltInGameFeaturePlugin(const TSharedRef<IPlug
 			}
 		}
 	}
+
+	UAssetManager::Get().PopBulkScanning();
 }
 
 void UGameFeaturesSubsystem::LoadBuiltInGameFeaturePlugins(FBuiltInPluginAdditionalFilters AdditionalFilter)
