Fix for non builtin GFPs not consistently having all their primary assets added to the asset manager.  Since these plugins are loaded from a previously unknown package root , we need synchronously the package paths for them.

diff --git a/Engine/Plugins/Experimental/GameFeatures/Source/GameFeatures/Private/GameFeaturesSubsystem.cpp b/Engine/Plugins/Experimental/GameFeatures/Source/GameFeatures/Private/GameFeaturesSubsystem.cpp
--- a/Engine/Plugins/Experimental/GameFeatures/Source/GameFeatures/Private/GameFeaturesSubsystem.cpp
+++ b/Engine/Plugins/Experimental/GameFeatures/Source/GameFeatures/Private/GameFeaturesSubsystem.cpp
@@ -313,17 +313,25 @@ void UGameFeaturesSubsystem::UnloadGameFeatureData(const UGameFeatureData* GameF
 
 void UGameFeaturesSubsystem::AddGameFeatureToAssetManager(const UGameFeatureData* GameFeatureToAdd, const FString& PluginName, TArray<FName>& OutNewPrimaryAssetTypes)
 {
-	// JMarcus TODO: FPackageName::RegisterMountPoint
-	// Also need to test cooked project loading to see if asset reg can find primary assets
-
 	check(GameFeatureToAdd);
 	FString PluginRootPath = TEXT("/") + PluginName + TEXT("/");
 	UAssetManager& LocalAssetManager = UAssetManager::Get();
+	IAssetRegistry& LocalAssetRegistry = LocalAssetManager.GetAssetRegistry();
+
+	// @TODO: HACK - There is no guarantee that the plugin mount point was added before inte initial asset scan.
+	// If not, ScanPathsForPrimaryAssets will fail to find primary assets without a syncronous scan.
+	// A proper fix for this would be to handle all the primary asset discovery internally ins the asset manager 
+	// instead of doing it here.
+	// We just mounted the folder that contains these primary assets and the editor background scan may not
+	// not be finished by the time this is called, but a rescan will happen later in OnAssetRegistryFilesLoaded 
+	// as long as LocalAssetRegistry.IsLoadingAssets() is true.
+	const bool bForceSynchronousScan = !LocalAssetRegistry.IsLoadingAssets();
 
 	LocalAssetManager.PushBulkScanning();
 
 	for (FPrimaryAssetTypeInfo TypeInfo : GameFeatureToAdd->GetPrimaryAssetTypesToScan())
 	{
+		// @TODO: we shouldn't be accessing private data here. Need a better way to do this
 		for (FDirectoryPath& Path : TypeInfo.Directories)
 		{
 			// Convert plugin-relative paths to full package paths
@@ -338,7 +346,6 @@ void UGameFeaturesSubsystem::AddGameFeatureToAssetManager(const UGameFeatureData
 
 		FPrimaryAssetTypeInfo ExistingAssetTypeInfo;
 		const bool bAlreadyExisted = LocalAssetManager.GetPrimaryAssetTypeInfo(FPrimaryAssetType(TypeInfo.PrimaryAssetType), /*out*/ ExistingAssetTypeInfo);
-		const bool bForceSynchronousScan = false; // We just mounted the folder that contains these primary assets and the editor background scan is not going to be finished by the time this is called, but a rescan will happen later in OnAssetRegistryFilesLoaded
 		LocalAssetManager.ScanPathsForPrimaryAssets(TypeInfo.PrimaryAssetType, TypeInfo.AssetScanPaths, TypeInfo.AssetBaseClassLoaded, TypeInfo.bHasBlueprintClasses, TypeInfo.bIsEditorOnly, bForceSynchronousScan);
 
 		if (!bAlreadyExisted)
