Fixes incorrect tooltip text for Game Feature Plugin details view.

#rnx

diff --git a/Engine/Plugins/Experimental/GameFeatures/Source/GameFeaturesEditor/Private/SGameFeatureStateWidget.cpp b/Engine/Plugins/Experimental/GameFeatures/Source/GameFeaturesEditor/Private/SGameFeatureStateWidget.cpp
--- a/Engine/Plugins/Experimental/GameFeatures/Source/GameFeaturesEditor/Private/SGameFeatureStateWidget.cpp
+++ b/Engine/Plugins/Experimental/GameFeatures/Source/GameFeaturesEditor/Private/SGameFeatureStateWidget.cpp
@@ -34,7 +34,7 @@ void SGameFeatureStateWidget::Construct(const FArguments& InArgs)
 				.ToolTip(LOCTEXT("SwitchToRegisteredTooltip", "Attempt to change the current state of this game feature to Registered.\n\nRegistered means that the assets in the plugin are known, but have not yet been loaded, except a few for discovery reasons (and it is not actively affecting gameplay yet)."))
 			+ SSegmentedControl<EGameFeaturePluginState>::Slot(EGameFeaturePluginState::Loaded)
 				.Text(GetDisplayNameOfState(EGameFeaturePluginState::Loaded))
-				.ToolTip(LOCTEXT("SwitchToLoadedTooltip", "Attempt to change the current state of this game feature to Loaded.\n\nLoaded means that the plugin is in local storage (i.e., it is on the hard drive) but it has not been registered, loaded, or activated yet."))
+				.ToolTip(LOCTEXT("SwitchToLoadedTooltip", "Attempt to change the current state of this game feature to Loaded.\n\nLoaded means that the plugin is loaded into memory and registered with some game systems, but not yet active and not affecting gameplay."))
 			+ SSegmentedControl<EGameFeaturePluginState>::Slot(EGameFeaturePluginState::Active)
 				.Text(GetDisplayNameOfState(EGameFeaturePluginState::Active))
 				.ToolTip(LOCTEXT("SwitchToActiveTooltip", "Attempt to change the current state of this game feature to Active.\n\nActive means that the plugin is fully loaded and active. It is affecting the game."))
