Fix GFP mounting state when a plugin with the same name is already registered

diff --git a/Engine/Plugins/Experimental/GameFeatures/Source/GameFeatures/Private/GameFeaturePluginStateMachine.cpp b/Engine/Plugins/Experimental/GameFeatures/Source/GameFeatures/Private/GameFeaturePluginStateMachine.cpp
--- a/Engine/Plugins/Experimental/GameFeatures/Source/GameFeatures/Private/GameFeaturePluginStateMachine.cpp
+++ b/Engine/Plugins/Experimental/GameFeatures/Source/GameFeatures/Private/GameFeaturePluginStateMachine.cpp
@@ -1052,10 +1052,18 @@ struct FGameFeaturePluginState_Mounting : public FGameFeaturePluginState
 		checkf(FPaths::GetExtension(StateProperties.PluginInstalledFilename) == TEXT("uplugin"), TEXT("PluginInstalledFilename must have a uplugin extension. PluginURL: %s"), *StateProperties.PluginURL);
 
 		// refresh the plugins list to let the plugin manager know about it
-		TSharedPtr<IPlugin> MaybePlugin = IPluginManager::Get().FindPlugin(StateProperties.PluginName);
+		const TSharedPtr<IPlugin> MaybePlugin = IPluginManager::Get().FindPlugin(StateProperties.PluginName);
 		const bool bNeedsPluginMount = (MaybePlugin == nullptr || MaybePlugin->GetDescriptor().bExplicitlyLoaded);
 
-		if (MaybePlugin == nullptr)
+		if (MaybePlugin)
+		{
+			if (!FPaths::IsSamePath(MaybePlugin->GetDescriptorFileName(), StateProperties.PluginInstalledFilename))
+			{
+				StateStatus.SetTransitionError(EGameFeaturePluginState::ErrorMounting, UE::GameFeatures::StateMachineErrorNamespace + TEXT("Plugin_Name_Already_In_Use"));
+				return;
+			}
+		}
+		else
 		{
 			const bool bAddedPlugin = IPluginManager::Get().AddToPluginsList(StateProperties.PluginInstalledFilename);
 			if (!bAddedPlugin)
