Fix an issue where delegates would fire before the incrementer was bound. Added an ensure to catch if it happens again.
Fixed an issue where components would be destroyed when weak lambdas would fire dropping a call

diff --git a/Engine/Plugins/Experimental/GameFeatures/Source/GameFeatures/Private/GameFeaturePluginStateMachine.cpp b/Engine/Plugins/Experimental/GameFeatures/Source/GameFeatures/Private/GameFeaturePluginStateMachine.cpp
--- a/Engine/Plugins/Experimental/GameFeatures/Source/GameFeatures/Private/GameFeaturePluginStateMachine.cpp
+++ b/Engine/Plugins/Experimental/GameFeatures/Source/GameFeatures/Private/GameFeaturePluginStateMachine.cpp
@@ -1559,6 +1559,7 @@ struct FGameFeaturePluginState_Deactivating : public FGameFeaturePluginState
 	void OnPauserCompleted()
 	{
 		check(IsInGameThread());
+		ensure(NumExpectedPausers != INDEX_NONE);
 		++NumObservedPausers;
 
 		if (NumObservedPausers == NumExpectedPausers)
@@ -1591,7 +1592,7 @@ struct FGameFeaturePluginState_Deactivating : public FGameFeaturePluginState
 
 		if (NumExpectedPausers == NumObservedPausers)
 		{
-			if(!bRequestedGC && StateProperties.Destination.MaxState == EGameFeaturePluginState::Loaded)
+			if (!bRequestedGC && StateProperties.Destination.MaxState == EGameFeaturePluginState::Loaded)
 			{
 				// If we aren't going farther than Loaded, GC now
 				// otherwise we will defer until closer to our destination state
