Fixed issue with game feature plugin manager not calling callback on change desired state if the state change was cancelled.

diff --git a/Engine/Plugins/Experimental/GameFeatures/Source/GameFeatures/Private/GameFeaturesSubsystem.cpp b/Engine/Plugins/Experimental/GameFeatures/Source/GameFeatures/Private/GameFeaturesSubsystem.cpp
--- a/Engine/Plugins/Experimental/GameFeatures/Source/GameFeatures/Private/GameFeaturesSubsystem.cpp
+++ b/Engine/Plugins/Experimental/GameFeatures/Source/GameFeatures/Private/GameFeaturesSubsystem.cpp
@@ -1157,20 +1157,20 @@ void UGameFeaturesSubsystem::ChangeGameFeatureDestination(UGameFeaturePluginStat
 
 			if (!ensure(bSetDestination))
 			{
-				UE_LOG(LogGameFeatures, Verbose, TEXT("ChangeGameFeatureDestination: Failed to set Game Feature %s Destination State to [%s, %s]"), *Machine->GetGameFeatureName(), *UE::GameFeatures::ToString(StateRange.MinState), *UE::GameFeatures::ToString(StateRange.MaxState));
+				UE_LOG(LogGameFeatures, Warning, TEXT("ChangeGameFeatureDestination: Failed to set Game Feature %s Destination State to [%s, %s]"), *Machine->GetGameFeatureName(), *UE::GameFeatures::ToString(StateRange.MinState), *UE::GameFeatures::ToString(StateRange.MaxState));
 
 				CompleteDelegate.ExecuteIfBound(MakeError(TEXT("GameFeaturePlugin.StateMachine.State_Currently_Unreachable")));
 			}
 			else
 			{
-				UE_LOG(LogGameFeatures, Verbose, TEXT("ChangeGameFeatureDestination: OnCanceled, set Game Feature %s Destination State to [%s, %s]"), *Machine->GetGameFeatureName(), *UE::GameFeatures::ToString(StateRange.MinState), *UE::GameFeatures::ToString(StateRange.MaxState));
+				UE_LOG(LogGameFeatures, Display, TEXT("ChangeGameFeatureDestination: OnCanceled, set Game Feature %s Destination State to [%s, %s]"), *Machine->GetGameFeatureName(), *UE::GameFeatures::ToString(StateRange.MinState), *UE::GameFeatures::ToString(StateRange.MaxState));
 			}
 		};
 
 		const bool bCancelPending = Machine->TryCancel(FGameFeatureStateTransitionCanceled::CreateWeakLambda(this, MoveTemp(OnCanceled)));
 		if (!ensure(bCancelPending))
 		{
-			UE_LOG(LogGameFeatures, Verbose, TEXT("ChangeGameFeatureDestination: Failed to cancel Game Feature %s"), *Machine->GetGameFeatureName());
+			UE_LOG(LogGameFeatures, Warning, TEXT("ChangeGameFeatureDestination: Failed to cancel Game Feature %s"), *Machine->GetGameFeatureName());
 
 			CompleteDelegate.ExecuteIfBound(MakeError(TEXT("GameFeaturePlugin.StateMachine.State_Currently_Unreachable_Cancel_Failed")));
 		}
