#UE Avoid needlessly scanning paths in the asset registry when loading a GameFeatureData. This was flushing the scan queue and causing slow editor load time in some cases. Also, only call LoadGameFeatureData once since the new method of loading does not immediately return an invalid handle if the file does not exist (its an async operation which does this later).

diff --git a/Engine/Plugins/Experimental/GameFeatures/Source/GameFeatures/Private/GameFeaturePluginStateMachine.cpp b/Engine/Plugins/Experimental/GameFeatures/Source/GameFeatures/Private/GameFeaturePluginStateMachine.cpp
--- a/Engine/Plugins/Experimental/GameFeatures/Source/GameFeatures/Private/GameFeaturePluginStateMachine.cpp
+++ b/Engine/Plugins/Experimental/GameFeatures/Source/GameFeatures/Private/GameFeaturePluginStateMachine.cpp
@@ -1365,8 +1365,12 @@ struct FGameFeaturePluginState_Registering : public FGameFeaturePluginState
 			BackupGameFeatureDataPath = OverridePath;
 		}
 		
-		TSharedPtr<FStreamableHandle> GameFeatureDataHandle = UGameFeaturesSubsystem::LoadGameFeatureData(PreferredGameFeatureDataPath);
-		if (!GameFeatureDataHandle.IsValid())
+		TSharedPtr<FStreamableHandle> GameFeatureDataHandle;
+		if (FPackageName::DoesPackageExist(PreferredGameFeatureDataPath))
+		{
+			GameFeatureDataHandle = UGameFeaturesSubsystem::LoadGameFeatureData(PreferredGameFeatureDataPath);
+		}
+		else
 		{
 			GameFeatureDataHandle = UGameFeaturesSubsystem::LoadGameFeatureData(BackupGameFeatureDataPath);
 		}
diff --git a/Engine/Plugins/Experimental/GameFeatures/Source/GameFeatures/Private/GameFeaturesSubsystem.cpp b/Engine/Plugins/Experimental/GameFeatures/Source/GameFeatures/Private/GameFeaturesSubsystem.cpp
--- a/Engine/Plugins/Experimental/GameFeatures/Source/GameFeatures/Private/GameFeaturesSubsystem.cpp
+++ b/Engine/Plugins/Experimental/GameFeatures/Source/GameFeatures/Private/GameFeaturesSubsystem.cpp
@@ -244,65 +244,7 @@ void UGameFeaturesSubsystem::OnAssetManagerCreated()
 
 TSharedPtr<FStreamableHandle> UGameFeaturesSubsystem::LoadGameFeatureData(const FString& GameFeatureToLoad)
 {
-	UAssetManager& LocalAssetManager = UAssetManager::Get();
-	IAssetRegistry& LocalAssetRegistry = LocalAssetManager.GetAssetRegistry();
-
-#if WITH_EDITOR
-	const FString GameFeaturePackageName = FPackageName::ObjectPathToPackageName(GameFeatureToLoad);
-	LocalAssetRegistry.ScanFilesSynchronous({ FPackageName::LongPackageNameToFilename(GameFeaturePackageName, FPackageName::GetAssetPackageExtension()) });
-#endif // WITH_EDITOR
-
-	FAssetData GameFeatureAssetData;
-	{
-		FARFilter ArFilter;
-		ArFilter.ObjectPaths.Add(FName(*GameFeatureToLoad));
-		ArFilter.ClassPaths.Add(UGameFeatureData::StaticClass()->GetClassPathName());
-		ArFilter.bRecursiveClasses = true;
-#if !WITH_EDITOR
-		ArFilter.bIncludeOnlyOnDiskAssets = true;
-#endif //if !WITH_EDITOR
-
-		LocalAssetRegistry.EnumerateAssets(ArFilter, [&GameFeatureAssetData](const FAssetData& AssetData)
-		{
-			GameFeatureAssetData = AssetData;
-			return false;
-		});
-	}
-
-	if (GameFeatureAssetData.IsValid())
-	{
-		FPrimaryAssetId AssetId = GameFeatureAssetData.GetPrimaryAssetId();
-
-#if WITH_EDITOR
-		// Support for pre-primary data asset game feature data, or game feature data copied from another plugin without being resaved.
-		FString PluginRoot;
-		FString ExpectedPluginRoot = FString::Printf(TEXT("/%s/"), *AssetId.PrimaryAssetName.ToString());
-		if (!AssetId.IsValid() || (UAssetManager::GetContentRootPathFromPackageName(GameFeaturePackageName, PluginRoot) && (PluginRoot != ExpectedPluginRoot)))
-		{
-			if (UObject* LoadedObject = GameFeatureAssetData.GetAsset())
-			{
-				AssetId = LoadedObject->GetPrimaryAssetId();
-				GameFeatureAssetData = FAssetData(LoadedObject);
-			}
-		}
-#endif // WITH_EDITOR
-
-		// Add the GameFeatureData itself to the primary asset list
-		LocalAssetManager.RegisterSpecificPrimaryAsset(AssetId, GameFeatureAssetData);
-
-		// LoadPrimaryAsset will return a null handle if the AssetID is already loaded. Check if there is an existing handle first.
-		TSharedPtr<FStreamableHandle> ReturnHandle = LocalAssetManager.GetPrimaryAssetHandle(AssetId);
-		if (ReturnHandle.IsValid())
-		{
-			return ReturnHandle;
-		}
-		else
-		{
-			return LocalAssetManager.LoadPrimaryAsset(AssetId);
-		}
-	}
-
-	return nullptr;
+	return UAssetManager::Get().GetStreamableManager().RequestAsyncLoad(FSoftObjectPath(GameFeatureToLoad));
 }
 
 void UGameFeaturesSubsystem::UnloadGameFeatureData(const UGameFeatureData* GameFeatureToUnload)
@@ -318,6 +260,9 @@ void UGameFeaturesSubsystem::AddGameFeatureToAssetManager(const UGameFeatureData
 	UAssetManager& LocalAssetManager = UAssetManager::Get();
 	IAssetRegistry& LocalAssetRegistry = LocalAssetManager.GetAssetRegistry();
 
+	// Add the GameFeatureData itself to the primary asset list
+	LocalAssetManager.RegisterSpecificPrimaryAsset(GameFeatureToAdd->GetPrimaryAssetId(), FAssetData(GameFeatureToAdd));
+
 	// @TODO: HACK - There is no guarantee that the plugin mount point was added before inte initial asset scan.
 	// If not, ScanPathsForPrimaryAssets will fail to find primary assets without a syncronous scan.
 	// A proper fix for this would be to handle all the primary asset discovery internally ins the asset manager 
