Fix editor UI for adding the first game feature plugin to a project, validation would always fail if Plugins/GameFeatures didn't already exist

diff --git a/Engine/Plugins/Experimental/GameFeatures/Source/GameFeatures/Private/GameFeaturesSubsystemSettings.cpp b/Engine/Plugins/Experimental/GameFeatures/Source/GameFeatures/Private/GameFeaturesSubsystemSettings.cpp
--- a/Engine/Plugins/Experimental/GameFeatures/Source/GameFeatures/Private/GameFeaturesSubsystemSettings.cpp
+++ b/Engine/Plugins/Experimental/GameFeatures/Source/GameFeatures/Private/GameFeaturesSubsystemSettings.cpp
@@ -18,10 +18,17 @@ bool UGameFeaturesSubsystemSettings::IsValidGameFeaturePlugin(const FString& Plu
 	// Build the cache of game feature plugin folders the first time this is called
 	if (BuiltInGameFeaturePluginsFolders.IsEmpty())
 	{
+		// Get all the existing game feature paths
 		BuiltInGameFeaturePluginsFolders.Append(
 			FPaths::GetExtensionDirs(FPaths::ProjectDir(), FPaths::Combine(TEXT("Plugins"), TEXT("GameFeatures")))
 		);
 
+		// The base directory may not exist yet, add it if empty
+		if (BuiltInGameFeaturePluginsFolders.IsEmpty())
+		{
+			BuiltInGameFeaturePluginsFolders.Add(FPaths::Combine(FPaths::ProjectDir(), TEXT("Plugins"), TEXT("GameFeatures")));
+		}
+
 		for (FString& BuiltInFolder : BuiltInGameFeaturePluginsFolders)
 		{
 			BuiltInFolder = FPaths::ConvertRelativePathToFull(BuiltInFolder);
