#UE Fix GameFeatureData bundles not being registered in cooked builds

diff --git a/Engine/Plugins/Experimental/GameFeatures/Source/GameFeatures/Private/GameFeaturesSubsystem.cpp b/Engine/Plugins/Experimental/GameFeatures/Source/GameFeatures/Private/GameFeaturesSubsystem.cpp
--- a/Engine/Plugins/Experimental/GameFeatures/Source/GameFeatures/Private/GameFeaturesSubsystem.cpp
+++ b/Engine/Plugins/Experimental/GameFeatures/Source/GameFeatures/Private/GameFeaturesSubsystem.cpp
@@ -261,7 +261,13 @@ void UGameFeaturesSubsystem::AddGameFeatureToAssetManager(const UGameFeatureData
 	IAssetRegistry& LocalAssetRegistry = LocalAssetManager.GetAssetRegistry();
 
 	// Add the GameFeatureData itself to the primary asset list
+#if WITH_EDITOR
+	// In the editor, we may not have scanned the FAssetData yet if during startup, but that is fine because we can gather bundles from the object itself, so just create the FAssetData from the object
 	LocalAssetManager.RegisterSpecificPrimaryAsset(GameFeatureToAdd->GetPrimaryAssetId(), FAssetData(GameFeatureToAdd));
+#else
+	// In non-editor, the asset bundle data is compiled out, so it must be gathered from the asset registry instead
+	LocalAssetManager.RegisterSpecificPrimaryAsset(GameFeatureToAdd->GetPrimaryAssetId(), LocalAssetRegistry.GetAssetByObjectPath(FName(*GameFeatureToAdd->GetPathName()), true));
+#endif // WITH_EDITOR
 
 	// @TODO: HACK - There is no guarantee that the plugin mount point was added before inte initial asset scan.
 	// If not, ScanPathsForPrimaryAssets will fail to find primary assets without a syncronous scan.
