Callsite upgrade pass for TObjectPtr use on plugins used when building UnrealEditor with "-allmodules".

diff --git a/Engine/Plugins/Experimental/GameFeatures/Source/GameFeatures/Private/GameFeaturesSubsystem.cpp b/Engine/Plugins/Experimental/GameFeatures/Source/GameFeatures/Private/GameFeaturesSubsystem.cpp
--- a/Engine/Plugins/Experimental/GameFeatures/Source/GameFeatures/Private/GameFeaturesSubsystem.cpp
+++ b/Engine/Plugins/Experimental/GameFeatures/Source/GameFeatures/Private/GameFeaturesSubsystem.cpp
@@ -1291,7 +1291,7 @@ TSet<FString> UGameFeaturesSubsystem::GetActivePluginNames() const
 {
 	TSet<FString> ActivePluginNames;
 
-	for (const TPair<FString, UGameFeaturePluginStateMachine*>& Pair : GameFeaturePluginStateMachines)
+	for (const UE_TRANSITIONAL_OBJECT_PTR_TEMPLATE2_ARG2(TPair, FString, UGameFeaturePluginStateMachine)& Pair : GameFeaturePluginStateMachines)
 	{
 		UGameFeaturePluginStateMachine* StateMachine = Pair.Value;
 		if (StateMachine->GetCurrentState() == EGameFeaturePluginState::Active &&
