GameFeaturePlugin Changes:

Changes to the core way a Plugin is stored and identified:
- Used to be a raw FString PluginURL but this causes problem with InstallBundleProtocol GameFeatures where we want to encode flags in the URL. Individual flags should not represent individual Plugins.
- Now its a new FGameFeaturePluginIdentifier that is parsed from the URL. This encodes URLs by the "important" information (basically just the UPlugin Descriptor Location stripped of the Protocol and all Flags)

Changes to UE::GameFeatures::FResult and FGameFeaturePluginState's use of FResult
- Added new OptionalErrorText that can be set to describe valuable lower level error messages in the GameFeaturePluginStates themselves instead of the requesting system needing to have detailed knowledge about what is going on under the hood to generate good localized error messages.
- New functions in FGameFeaturePluginState GetErrorResult that help set better FResults with shortcuts for appending the State error namespace and package the new OptionalErrorText where appropriate.
- Added some new ErrorTexts to be returned in certain cases.

Added / Changed States
- Added new Uninstall state between StatusKnown -> Terminated that runs if we are an InstallBundle and the new
- Added new ErrorUninstalling to handle errors in the above Uninstalling state
- Renamed old Uninstalling state to Releasing to better match what it is doing (Releasing BundleManager cache and not actually uninstalling data)
- Renamed ErrorInstalling to ErrorManagingData as its an Error state for both Releasing and Downloading

Misc Changes to UGameFeaturesSubsystem
- Renamed  UninstallGameFeaturePlugin to ReleaseGameFeaturePlugin to better clarify what it is doing now that there is a new Uninstall state
- Added a new bShouldCreateGameFeaturePluginIfNotFound version of TerminateGameFeaturePlugin. This allows us to create a GFP to terminate even if it wasn't previously requested. This seems weird at first, but its to allow us to Uninstall GFP's that weren't previously known by the GFP system but might have datat installed from a previous run of the application. Now to uninstall you would just set the Metadata bUninstallBeforeTerminate and call TerminateGameFeaturePlugin.
- Renamed / Implemented some new Delegate callbacks to reflect these changes.
- Added a new way to generate an InstallBundle Protocol URL that takes in a FInstallBundlePluginProtocolMetaData and uses it's ToString() to generate your URL. More convenient since you can just set all the MetaData options directly.
- Added IsGameFeaturePluginInstalled() to compliment existing IsGameFeaturePlugin... Registered, Loaded, and Active
- Added new OnGameFeaturePauseChange to IGameFeatureStateChangeObserver to track whenever a GameFeaturePlugin State pauses work without transitioning out of the state so it can be monitored and displayed. EX: Downloading paused because of cellular data prompt settings.

UpdateGameFeaturePluginURL
- New addition to the UGameFeaturesSubsystem that is big enough for it's own bullet point
- UpdateGameFeaturePluginURL that allows you to update the associated PluginURL for your GameFeaturePlugin. Used to change MetaData on your URL.
- Individual States can implement TryUpdateURLData to react to this changing data and determine what a successful update looks like. Default implementation just updates and verifies the MetaData.

FInstallBundlePluginProtocolMetaData changes
- bUninstallBeforeTerminate - Determines if we should run the Uninstall state when leaving StatusKnown -> Terminal and do StatusKnown -> Uninstall instead.
- InstallBundleFlags - What EInstallBundleRequest flags are passed into InstallBundleManager during our Downloading and Mounting stages
- bUserPauseDownload - Set/unset if the user wants to pause/resume the download during the Downloading state
- Now use ToString / FromString for parsing of these values instead of just having all the code in FGameFeaturePluginStateMachineProperties::ParseURL

#tests

diff --git a/Engine/Plugins/Experimental/GameFeatures/Source/GameFeatures/Private/GameFeaturePluginOperationResult.cpp b/Engine/Plugins/Experimental/GameFeatures/Source/GameFeatures/Private/GameFeaturePluginOperationResult.cpp
new file mode 100644
--- /dev/null
+++ b/Engine/Plugins/Experimental/GameFeatures/Source/GameFeatures/Private/GameFeaturePluginOperationResult.cpp
@@ -0,0 +1,24 @@
+// Copyright Epic Games, Inc. All Rights Reserved.
+
+#include "GameFeaturePluginOperationResult.h"
+
+
+namespace UE::GameFeatures
+{
+	FResult::FResult(FErrorCodeType ErrorCodeIn)
+		: ErrorCode(MoveTemp(ErrorCodeIn))
+		, OptionalErrorText()
+	{
+	}
+
+	FResult::FResult(FErrorCodeType ErrorCodeIn, FText ErrorTextIn)
+		: ErrorCode(MoveTemp(ErrorCodeIn))
+		, OptionalErrorText(MoveTemp(ErrorTextIn))
+	{
+	}
+
+	FString ToString(const FResult& Result)
+	{
+		return Result.HasValue() ? FString(TEXT("Success")) : (FString(TEXT("Failure, ErrorCode=")) + Result.GetError() + FString(TEXT(", OptionalErrorText=")) + Result.OptionalErrorText.ToString());
+	}
+}	// namespace UE::GameFeatures
\ No newline at end of file
diff --git a/Engine/Plugins/Experimental/GameFeatures/Source/GameFeatures/Private/GameFeaturePluginStateMachine.cpp b/Engine/Plugins/Experimental/GameFeatures/Source/GameFeatures/Private/GameFeaturePluginStateMachine.cpp
--- a/Engine/Plugins/Experimental/GameFeatures/Source/GameFeatures/Private/GameFeaturePluginStateMachine.cpp
+++ b/Engine/Plugins/Experimental/GameFeatures/Source/GameFeatures/Private/GameFeaturePluginStateMachine.cpp
@@ -48,12 +48,6 @@ namespace UE::GameFeatures
 		false,
 		TEXT("Dump reference chains for any detected plugin asset leaks."),
 		ECVF_Default);
-
-	FString ToString(const UE::GameFeatures::FResult& Result)
-	{
-		return Result.HasValue() ? FString(TEXT("Success")) : (FString(TEXT("Failure, ErrorCode=")) + Result.GetError());
-	}
-
 	#define GAME_FEATURE_PLUGIN_STATE_TO_STRING(inEnum, inText) case EGameFeaturePluginState::inEnum: return TEXT(#inEnum);
 	FString ToString(EGameFeaturePluginState InType)
 	{
@@ -143,13 +137,229 @@ namespace UE::GameFeatures
 		return nullptr;
 	}
 #undef GAME_FEATURE_PLUGIN_PROTOCOL_PREFIX
+
+	namespace CommonErrorCodes
+	{
+		const FText Generic_FatalError = NSLOCTEXT("GameFeatures", "ErrorCodes.GenericFatalError", "A fatal error has occurred installing the game feature. An update to the application may be needed. Please check for updates and restart the application.");
+		const FText Generic_ConnectionError = NSLOCTEXT("GameFeatures", "ErrorCodes.ConnectionGenericError", "An internet connection error has occurred. Please try again later.");
+		const FText Generic_MountError = NSLOCTEXT("GameFeatures", "ErrorCodes.MountGenericError", "An error has occurred loading data for this game feature. Please try again later.");
+
+		const FText BundleResult_NeedsUpdate = NSLOCTEXT("GameFeatures", "ErrorCodes.BundleResult.NeedsUpdate", "An application update is required to install this game feature. Please restart the application after downloading any required updates.");
+		const FText BundleResult_NeedsDiskSpace = NSLOCTEXT("GameFeatures", "ErrorCodes.BundleResult.NeedsDiskSpace", "You do not have enough disk space to install this game feature. Please try again after clearing up disk space.");
+		const FText BundleResult_DownloadCancelled = NSLOCTEXT("GameFeatures", "ErrorCodes.BundleResult.DownloadCancelled", "This game feature download was canceled.");
+
+		const FText ReleaseResult_Generic = NSLOCTEXT("GameFeatures", "ErrorCodes.ReleaseResult.Generic", "There was an error uninstalling the content for this game feature. Please restart the application and try again.");
+		const FText ReleaseResult_Cancelled = NSLOCTEXT("GameFeatures", "ErrorCodes.ReleaseResult.Cancelled", "This game feature uninstall was canceled.");
+
+		static const FText& GetErrorTextForBundleResult(EInstallBundleResult ErrorResult)
+		{
+			switch (ErrorResult)
+			{
+				//These errors mean an app update is available that we either don't have or failed to get.
+				case EInstallBundleResult::FailedPrereqRequiresLatestClient:
+				case EInstallBundleResult::FailedPrereqRequiresLatestContent:
+				{
+					return BundleResult_NeedsUpdate;
+				}
+
+
+				//These are generally unrecoverable and mean something is seriously wrong with the data for this build
+				case EInstallBundleResult::InitializationError:
+				case EInstallBundleResult::FailedCacheReserve:
+				{
+					return Generic_FatalError;
+				}
+
+				//All of these are indicative of not having enough space to install the required files
+				case EInstallBundleResult::InstallerOutOfDiskSpaceError:
+				case EInstallBundleResult::ManifestArchiveError:
+				{
+					return BundleResult_NeedsDiskSpace;
+				}
+
+				case EInstallBundleResult::UserCancelledError:
+				{
+					return BundleResult_DownloadCancelled;
+				}
+
+				//Intentionally just show generic error for all these cases
+				case EInstallBundleResult::InstallError:
+				case EInstallBundleResult::ConnectivityError:
+				case EInstallBundleResult::InitializationPending:
+				{
+					return Generic_ConnectionError;
+				}
+
+				//Show generic error for anything missing but log an error
+				default:
+				{
+					UE_LOG(LogGameFeatures, Error, TEXT("Missing error text for EInstallBundleResult %s"), *LexToString(ErrorResult));
+					return Generic_ConnectionError;
+				}
+			}
+		}
+
+		static const FText& GetErrorTextForReleaseResult(EInstallBundleReleaseResult ErrorResult)
+		{
+			switch (ErrorResult)
+			{
+				case (EInstallBundleReleaseResult::UserCancelledError):
+				{
+					return ReleaseResult_Cancelled;
+				}
+			
+				case (EInstallBundleReleaseResult::ManifestArchiveError):
+				{
+					return ReleaseResult_Generic;
+				}
+
+				default:
+				{
+					//Show generic error for anything missing but log an error
+					UE_LOG(LogGameFeatures, Error, TEXT("Missing error text for EInstallBundleReleaseResult %s"), *LexToString(ErrorResult));
+					return ReleaseResult_Generic;
+				}
+			}
+		}
+	};
 }
 
+//Enum describing the options for what data can be inside the InstallBundle protocol URL's metadata
+//Used to aid in parsing / creating the URL for InstallBundle protocol GameFeaturePlugins
+#define INSTALL_BUNDLE_PROTOCOL_URL_OPTIONS_LIST(XOPTION)					\
+	XOPTION(Bundles,					TEXT("Bundles"))					\
+	XOPTION(Version,					TEXT("V"))							\
+	XOPTION(Flags,						TEXT("Flags"))						\
+	XOPTION(UninstallBeforeTerminate,	TEXT("UninstallBeforeTerminate"))	\
+	XOPTION(UserPauseDownload,			TEXT("Paused"))
+
+#define INSTALL_BUNDLE_PROTOCOL_URL_OPTIONS_ENUM(inEnum, inString) inEnum,
+enum class EGameFeatureInstallBundleProtocolOptions : uint8
+{
+	INSTALL_BUNDLE_PROTOCOL_URL_OPTIONS_LIST(INSTALL_BUNDLE_PROTOCOL_URL_OPTIONS_ENUM)
+	Count,
+};
+#undef INSTALL_BUNDLE_PROTOCOL_URL_OPTIONS_ENUM
+
+#define INSTALL_BUNDLE_PROTOCOL_URL_OPTIONS_LEX_TO_STRING(inEnum, inString) \
+    case(EGameFeatureInstallBundleProtocolOptions::inEnum):                 \
+    {                                                                       \
+        return inString;                                                    \
+    }                                                                       
+
+FString LexToString(EGameFeatureInstallBundleProtocolOptions FeatureIn)
+{
+	switch (FeatureIn)
+	{
+		INSTALL_BUNDLE_PROTOCOL_URL_OPTIONS_LIST(INSTALL_BUNDLE_PROTOCOL_URL_OPTIONS_LEX_TO_STRING)
+		
+		default:
+		{
+			ensureAlwaysMsgf(false, TEXT("Logic error causing a missing LexToString value for EGameFeatureInstallBundleProtocolOptions:%d"), static_cast<uint8>(FeatureIn));
+			return ("ERROR_UNSUPPORTED_ENUM");
+		}
+	}
+}
+#undef INSTALL_BUNDLE_PROTOCOL_URL_OPTIONS_LEX_TO_STRING
+#undef INSTALL_BUNDLE_PROTOCOL_URL_OPTIONS_LIST
+
+void FGameFeaturePluginStateStatus::SetTransition(EGameFeaturePluginState InTransitionToState)
+{
+	TransitionToState = InTransitionToState;
+	TransitionResult.ErrorCode = MakeValue();
+	TransitionResult.OptionalErrorText = FText();
+}
+
+void FGameFeaturePluginStateStatus::SetTransitionError(EGameFeaturePluginState TransitionToErrorState, UE::GameFeatures::FResult TransitionResultIn)
+{
+	TransitionToState = TransitionToErrorState;
+	if (ensureAlwaysMsgf(TransitionResultIn.HasError(), TEXT("Invalid call to SetTransitionError with an FResult that isn't an error! TransitionToErrorState: %s"), *LexToString(TransitionToErrorState)))
+	{
+		TransitionResult = MoveTemp(TransitionResultIn);
+	}
+	else
+	{
+		//Logic error using a non-error FResult, so just generate a general error to keep the SetTransitionError intent
+		TransitionResult = MakeError(TEXT("Invalid_Transition_Error"));
+	}
+}
+
+EGameFeatureInstallBundleProtocolOptions LexFromString(const FString& StringIn)
+{
+	if (StringIn.IsEmpty())
+	{
+		return EGameFeatureInstallBundleProtocolOptions::Count;
+	}
+
+	for (uint8 EnumIndex = 0; EnumIndex < static_cast<uint8>(EGameFeatureInstallBundleProtocolOptions::Count); ++EnumIndex)
+	{
+		EGameFeatureInstallBundleProtocolOptions GameFeatureToCheck = static_cast<EGameFeatureInstallBundleProtocolOptions>(EnumIndex);
+		if (StringIn.Equals(LexToString(GameFeatureToCheck)))
+		{
+			return GameFeatureToCheck;
+		}
+	}
+
+	return  EGameFeatureInstallBundleProtocolOptions::Count;
+}
+
+UE::GameFeatures::FResult FGameFeaturePluginState::GetErrorResult(const FString& ErrorCode, const FText OptionalErrorText/*= FText()*/) const
+{
+	return GetErrorResult(TEXT(""), ErrorCode, OptionalErrorText);
+}
+
+UE::GameFeatures::FResult FGameFeaturePluginState::GetErrorResult(const FString& ErrorNamespaceAddition, const FString& ErrorCode, const FText OptionalErrorText/*= FText()*/) const
+{
+	const FString StateName = LexToString(UGameFeaturesSubsystem::Get().GetPluginState(StateProperties.PluginIdentifier.GetFullPluginURL()));
+	const FString ErrorCodeEnding = ErrorNamespaceAddition.IsEmpty() ? ErrorCode : ErrorNamespaceAddition + ErrorCode;
+	const FString CompleteErrorCode = (UE::GameFeatures::StateMachineErrorNamespace + StateName + ErrorCodeEnding);
+	return UE::GameFeatures::FResult(MakeError(CompleteErrorCode), OptionalErrorText);
+}
+
+UE::GameFeatures::FResult FGameFeaturePluginState::GetErrorResult(const FString& ErrorNamespaceAddition, const EInstallBundleResult ErrorResult) const
+{
+	UE::GameFeatures::FResult BaseResult = GetErrorResult(ErrorNamespaceAddition, LexToString(ErrorResult));
+	BaseResult.OptionalErrorText = UE::GameFeatures::CommonErrorCodes::GetErrorTextForBundleResult(ErrorResult);
+	return MoveTemp(BaseResult);
+}
+
+UE::GameFeatures::FResult FGameFeaturePluginState::GetErrorResult(const FString& ErrorNamespaceAddition, const EInstallBundleReleaseResult ErrorResult) const
+{
+	UE::GameFeatures::FResult BaseResult = GetErrorResult(ErrorNamespaceAddition, LexToString(ErrorResult));
+	BaseResult.OptionalErrorText = UE::GameFeatures::CommonErrorCodes::GetErrorTextForReleaseResult(ErrorResult);
+	return MoveTemp(BaseResult);
+}
+
+
 FGameFeaturePluginState::~FGameFeaturePluginState()
 {
 	CleanupDeferredUpdateCallbacks();
 }
 
+bool FGameFeaturePluginState::TryUpdateURLData(const FString& NewPluginURL)
+{
+	//Early out without any update if we already fully match the new URL
+	if (StateProperties.PluginIdentifier.ExactMatchesURL(NewPluginURL))
+	{
+		return false;
+	}
+
+	//Create copied StateMachineProperties data and try and reset it by parsing the new data
+	FGameFeaturePluginIdentifier NewIdentifier = FGameFeaturePluginIdentifier(NewPluginURL);
+	FGameFeaturePluginStateMachineProperties NewProperties = StateProperties;
+	NewProperties.PluginIdentifier = NewIdentifier;
+
+	if (!NewProperties.ParseURL() || !NewProperties.ValidateURLUpdate(StateProperties))
+	{
+		return false;
+	}
+
+	//Parse was successful so transfer to the newly parsed properties
+	StateProperties = MoveTemp(NewProperties);
+
+	return true;
+}
+
 FDestinationGameFeaturePluginState* FGameFeaturePluginState::AsDestinationState()
 {
 	FDestinationGameFeaturePluginState* Ret = nullptr;
@@ -217,6 +427,24 @@ void FGameFeaturePluginState::CleanupDeferredUpdateCallbacks() const
 	FCoreUObjectDelegates::GetPostGarbageCollect().RemoveAll(this);
 }
 
+bool FGameFeaturePluginState::ShouldVisitUninstallStateBeforeTerminal() const
+{
+	switch (StateProperties.GetPluginProtocol())
+	{
+		case (EGameFeaturePluginProtocol::InstallBundle):
+		{
+			//InstallBundleProtocol's have a MetaData that controlls if they uninstall currently
+			return StateProperties.ProtocolMetadata.GetSubtype<FInstallBundlePluginProtocolMetaData>().bUninstallBeforeTerminate;
+		}
+
+		//Default behavior is to just Terminate
+		default:
+		{
+			return false;
+		}
+	}
+}
+
 /*
 =========================================================
   States
@@ -241,10 +469,10 @@ struct FGameFeaturePluginState_Terminal : public FDestinationGameFeaturePluginSt
 
 	virtual void BeginState() override
 	{
-		checkf(!bEnteredTerminalState, TEXT("Plugin entered terminal state more than once! %s"), *StateProperties.PluginURL);
+		checkf(!bEnteredTerminalState, TEXT("Plugin entered terminal state more than once! %s"), *StateProperties.PluginIdentifier.GetFullPluginURL());
 		bEnteredTerminalState = true;
 
-		UGameFeaturesSubsystem::Get().OnGameFeatureTerminating(StateProperties.PluginName, StateProperties.PluginURL);
+		UGameFeaturesSubsystem::Get().OnGameFeatureTerminating(StateProperties.PluginName, StateProperties.PluginIdentifier.GetFullPluginURL());
 	}
 };
 
@@ -262,7 +490,7 @@ struct FGameFeaturePluginState_UnknownStatus : public FDestinationGameFeaturePlu
 		{
 			StateStatus.SetTransition(EGameFeaturePluginState::CheckingStatus);
 
-			UGameFeaturesSubsystem::Get().OnGameFeatureCheckingStatus(StateProperties.PluginURL);
+			UGameFeaturesSubsystem::Get().OnGameFeatureCheckingStatus(StateProperties.PluginIdentifier.GetFullPluginURL());
 		}
 	}
 };
@@ -287,7 +515,7 @@ struct FGameFeaturePluginState_CheckingStatus : public FGameFeaturePluginState
 			bParsedURL = StateProperties.ParseURL();
 			if (!bParsedURL)
 			{
-				StateStatus.SetTransitionError(EGameFeaturePluginState::ErrorCheckingStatus, UE::GameFeatures::StateMachineErrorNamespace + TEXT("Bad_PluginURL"));
+				StateStatus.SetTransitionError(EGameFeaturePluginState::ErrorCheckingStatus, GetErrorResult(TEXT("Bad_PluginURL")));
 				return;
 			}
 		}
@@ -301,13 +529,13 @@ struct FGameFeaturePluginState_CheckingStatus : public FGameFeaturePluginState
 			TSharedPtr<IInstallBundleManager> BundleManager = IInstallBundleManager::GetPlatformInstallBundleManager();
 			if (BundleManager == nullptr)
 			{
-				StateStatus.SetTransitionError(EGameFeaturePluginState::ErrorCheckingStatus, UE::GameFeatures::StateMachineErrorNamespace + TEXT("BundleManager_Was_Null"));
+				StateStatus.SetTransitionError(EGameFeaturePluginState::ErrorCheckingStatus, GetErrorResult(TEXT("BundleManager_Null")));
 				return;
 			}
 
 			if (BundleManager->GetInitState() == EInstallBundleManagerInitState::Failed)
 			{
-				StateStatus.SetTransitionError(EGameFeaturePluginState::ErrorCheckingStatus, UE::GameFeatures::StateMachineErrorNamespace + TEXT("BundleManager_Failed_Init"));
+				StateStatus.SetTransitionError(EGameFeaturePluginState::ErrorCheckingStatus, GetErrorResult(TEXT("BundleManager_Failed_Init")));
 				return;
 			}
 
@@ -318,12 +546,12 @@ struct FGameFeaturePluginState_CheckingStatus : public FGameFeaturePluginState
 				return;
 			}
 
-			const TArray<FName>& InstallBundles = StateProperties.ProtocolMetadata.GetSubtype<UE::GameFeatures::FInstallBundlePluginProtocolMetaData>().InstallBundles;
+			const TArray<FName>& InstallBundles = StateProperties.ProtocolMetadata.GetSubtype<FInstallBundlePluginProtocolMetaData>().InstallBundles;
 
 			TValueOrError<FInstallBundleCombinedInstallState, EInstallBundleResult> MaybeInstallState = BundleManager->GetInstallStateSynchronous(InstallBundles, false);
 			if (MaybeInstallState.HasError())
 			{
-				StateStatus.SetTransitionError(EGameFeaturePluginState::ErrorCheckingStatus, UE::GameFeatures::StateMachineErrorNamespace + TEXT("BundleManager_Failed_GetInstallState"));
+				StateStatus.SetTransitionError(EGameFeaturePluginState::ErrorCheckingStatus, GetErrorResult(TEXT("BundleManager_Failed_GetInstallState")));
 				return;
 			}
 
@@ -332,17 +560,17 @@ struct FGameFeaturePluginState_CheckingStatus : public FGameFeaturePluginState
 		}
 		else
 		{
-			StateStatus.SetTransitionError(EGameFeaturePluginState::ErrorCheckingStatus, UE::GameFeatures::StateMachineErrorNamespace + TEXT("Unknown_Protocol"));
+			StateStatus.SetTransitionError(EGameFeaturePluginState::ErrorCheckingStatus, GetErrorResult(TEXT("Unknown_Protocol")));
 			return;
 		}
 
 		if (!bIsAvailable)
 		{
-			StateStatus.SetTransitionError(EGameFeaturePluginState::ErrorUnavailable, UE::GameFeatures::StateMachineErrorNamespace + TEXT("Plugin_Unavailable"));
+			StateStatus.SetTransitionError(EGameFeaturePluginState::ErrorUnavailable, GetErrorResult(TEXT("Plugin_Unavailable")));
 			return;
 		}
 
-		UGameFeaturesSubsystem::Get().OnGameFeatureStatusKnown(StateProperties.PluginName, StateProperties.PluginURL);
+		UGameFeaturesSubsystem::Get().OnGameFeatureStatusKnown(StateProperties.PluginName, StateProperties.PluginIdentifier.GetFullPluginURL());
 		StateStatus.SetTransition(EGameFeaturePluginState::StatusKnown);
 	}
 };
@@ -389,7 +617,14 @@ struct FGameFeaturePluginState_StatusKnown : public FDestinationGameFeaturePlugi
 	{
 		if (StateProperties.Destination < EGameFeaturePluginState::StatusKnown)
 		{
-			StateStatus.SetTransition(EGameFeaturePluginState::Terminal);
+			if (ShouldVisitUninstallStateBeforeTerminal())
+			{
+				StateStatus.SetTransition(EGameFeaturePluginState::Uninstalling);	
+			}
+			else
+			{
+				StateStatus.SetTransition(EGameFeaturePluginState::Terminal);
+			}
 		}
 		else if (StateProperties.Destination > EGameFeaturePluginState::StatusKnown)
 		{
@@ -405,30 +640,58 @@ struct FGameFeaturePluginState_StatusKnown : public FDestinationGameFeaturePlugi
 	}
 };
 
-struct FGameFeaturePluginState_ErrorInstalling : public FErrorGameFeaturePluginState
+struct FGameFeaturePluginState_ErrorManagingData : public FErrorGameFeaturePluginState
 {
-	FGameFeaturePluginState_ErrorInstalling(FGameFeaturePluginStateMachineProperties& InStateProperties) : FErrorGameFeaturePluginState(InStateProperties) {}
+	FGameFeaturePluginState_ErrorManagingData(FGameFeaturePluginStateMachineProperties& InStateProperties) : FErrorGameFeaturePluginState(InStateProperties) {}
 
 	virtual void UpdateState(FGameFeaturePluginStateStatus& StateStatus) override
 	{
-		if (StateProperties.Destination < EGameFeaturePluginState::ErrorInstalling)
-		{
-			StateStatus.SetTransition(EGameFeaturePluginState::Uninstalling);
+		if (StateProperties.Destination < EGameFeaturePluginState::ErrorManagingData)
+		{			
+			StateStatus.SetTransition(EGameFeaturePluginState::Releasing);
 		}
-		else if (StateProperties.Destination > EGameFeaturePluginState::ErrorInstalling)
+		else if (StateProperties.Destination > EGameFeaturePluginState::ErrorManagingData)
 		{
 			StateStatus.SetTransition(EGameFeaturePluginState::Downloading);
 		}
 	}
 };
 
-struct FGameFeaturePluginState_Uninstalling : public FGameFeaturePluginState
+struct FGameFeaturePluginState_ErrorUninstalling : public FErrorGameFeaturePluginState
 {
-	FGameFeaturePluginState_Uninstalling(FGameFeaturePluginStateMachineProperties& InStateProperties)
-		: FGameFeaturePluginState(InStateProperties)
+	FGameFeaturePluginState_ErrorUninstalling(FGameFeaturePluginStateMachineProperties& InStateProperties) : FErrorGameFeaturePluginState(InStateProperties) {}
+
+	virtual void UpdateState(FGameFeaturePluginStateStatus& StateStatus) override
+	{
+		if (StateProperties.Destination < EGameFeaturePluginState::ErrorUninstalling)
+		{
+			if (ShouldVisitUninstallStateBeforeTerminal())
+			{
+				StateStatus.SetTransition(EGameFeaturePluginState::Uninstalling);
+			}
+			else
+			{
+				StateStatus.SetTransition(EGameFeaturePluginState::Terminal);
+			}
+		}
+		else if (StateProperties.Destination > EGameFeaturePluginState::ErrorUninstalling)
+		{
+			StateStatus.SetTransition(EGameFeaturePluginState::StatusKnown);
+		}
+	}
+};
+
+//Base class for our states that want to request to release data from InstallBundleManager
+struct FBaseDataReleaseGameFeaturePluginState : public FGameFeaturePluginState
+{
+	FBaseDataReleaseGameFeaturePluginState(FGameFeaturePluginStateMachineProperties& InStateProperties) 
+		: FGameFeaturePluginState(InStateProperties) 
 		, Result(MakeValue())
 	{}
 
+	virtual ~FBaseDataReleaseGameFeaturePluginState()
+	{}
+
 	UE::GameFeatures::FResult Result;
 	bool bWasDeleted = false;
 	TArray<FName> PendingBundles;
@@ -444,7 +707,7 @@ struct FGameFeaturePluginState_Uninstalling : public FGameFeaturePluginState
 
 		if (!Result.HasError() && BundleResult.Result != EInstallBundleReleaseResult::OK)
 		{
-			Result = MakeError(UE::GameFeatures::StateMachineErrorNamespace + TEXT("BundleManager_Uninstall_Failure_") + LexToString(BundleResult.Result));
+			Result = GetErrorResult(TEXT("BundleManager.OnRemove_Failed."), BundleResult.Result);
 		}
 
 		if (PendingBundles.Num() > 0)
@@ -462,26 +725,35 @@ struct FGameFeaturePluginState_Uninstalling : public FGameFeaturePluginState
 
 	virtual void BeginState() override
 	{
+		BeginRemoveRequest();
+	}
+
+	virtual void BeginRemoveRequest()
+	{
+		CleanUp();
+
 		Result = MakeValue();
 		bWasDeleted = false;
 
-		if (StateProperties.GetPluginProtocol() != EGameFeaturePluginProtocol::InstallBundle)
+		if (!ShouldReleaseContent())
 		{
 			bWasDeleted = true;
 			return;
 		}
 
 		TSharedPtr<IInstallBundleManager> BundleManager = IInstallBundleManager::GetPlatformInstallBundleManager();
+		check(BundleManager.IsValid());
 
-		const TArray<FName>& InstallBundles = StateProperties.ProtocolMetadata.GetSubtype<UE::GameFeatures::FInstallBundlePluginProtocolMetaData>().InstallBundles;
+		const TArray<FName>& InstallBundles = StateProperties.ProtocolMetadata.GetSubtype<FInstallBundlePluginProtocolMetaData>().InstallBundles;
 
-		EInstallBundleReleaseRequestFlags ReleaseFlags = EInstallBundleReleaseRequestFlags::RemoveFilesIfPossible;
+		EInstallBundleReleaseRequestFlags ReleaseFlags = GetReleaseRequestFlags();
 		TValueOrError<FInstallBundleRequestInfo, EInstallBundleResult> MaybeRequestInfo = BundleManager->RequestReleaseContent(InstallBundles, ReleaseFlags);
 
-		if (!MaybeRequestInfo.IsValid())
+		if (MaybeRequestInfo.HasError())
 		{
-			ensureMsgf(false, TEXT("Unable to enqueue uninstall for the PluginURL(%s) because %s"), *StateProperties.PluginURL, LexToString(MaybeRequestInfo.GetError()));
-			Result = MakeError(UE::GameFeatures::StateMachineErrorNamespace + TEXT("BundleManager_Cannot_Uninstall"));
+			ensureMsgf(false, TEXT("Unable to enqueue uninstall for the PluginURL(%s) because %s"), *StateProperties.PluginIdentifier.GetFullPluginURL(), LexToString(MaybeRequestInfo.GetError()));
+			Result = GetErrorResult(TEXT("BundleManager.Begin."), MaybeRequestInfo.GetError());
+
 			return;
 		}
 
@@ -489,8 +761,9 @@ struct FGameFeaturePluginState_Uninstalling : public FGameFeaturePluginState
 
 		if (EnumHasAnyFlags(RequestInfo.InfoFlags, EInstallBundleRequestInfoFlags::SkippedUnknownBundles))
 		{
-			ensureMsgf(false, TEXT("Unable to enqueue uninstall for the PluginURL(%s) because failed to resolve install bundles!"), *StateProperties.PluginURL);
-			Result = MakeError(UE::GameFeatures::StateMachineErrorNamespace + TEXT("BundleManager_Cannot_Resolve_InstallBundles_For_Release"));
+			ensureMsgf(false, TEXT("Unable to enqueue uninstall for the PluginURL(%s) because failed to resolve install bundles!"), *StateProperties.PluginIdentifier.GetFullPluginURL());
+			Result = GetErrorResult(TEXT("BundleManager.Begin."), TEXT("Resolve_Failed"), UE::GameFeatures::CommonErrorCodes::ReleaseResult_Generic);
+
 			return;
 		}
 
@@ -501,7 +774,7 @@ struct FGameFeaturePluginState_Uninstalling : public FGameFeaturePluginState
 		else
 		{
 			PendingBundles = MoveTemp(RequestInfo.BundlesEnqueued);
-			IInstallBundleManager::ReleasedDelegate.AddRaw(this, &FGameFeaturePluginState_Uninstalling::OnContentRemoved);
+			IInstallBundleManager::ReleasedDelegate.AddRaw(this, &FBaseDataReleaseGameFeaturePluginState::OnContentRemoved);
 		}
 	}
 
@@ -509,7 +782,7 @@ struct FGameFeaturePluginState_Uninstalling : public FGameFeaturePluginState
 	{
 		if (!Result.HasValue())
 		{
-			StateStatus.SetTransitionError(EGameFeaturePluginState::ErrorInstalling, Result.StealError());
+			StateStatus.SetTransitionError(GetFailureTransitionState(), Result);
 			return;
 		}
 
@@ -518,13 +791,115 @@ struct FGameFeaturePluginState_Uninstalling : public FGameFeaturePluginState
 			return;
 		}
 
-		StateStatus.SetTransition(EGameFeaturePluginState::StatusKnown);
+		StateStatus.SetTransition(GetSuccessTransitionState());
 	}
 
-	virtual void EndState() override
+	void CleanUp()
 	{
+		PendingBundles.Empty();
 		IInstallBundleManager::ReleasedDelegate.RemoveAll(this);
 	}
+
+	virtual void EndState() override
+	{
+		CleanUp();
+	}
+
+	/** Controls what check is done to determine if this state should run or not */
+	virtual bool ShouldReleaseContent() const
+	{
+		switch (StateProperties.GetPluginProtocol())
+		{
+			case (EGameFeaturePluginProtocol::InstallBundle):
+			{
+				return true;
+			}
+
+			default:
+			{
+				return false;
+			}
+		}
+	}
+
+	/** Determines what state you transition to in the event of a success or failure to release content */
+	virtual EGameFeaturePluginState GetSuccessTransitionState() const = 0;
+	virtual EGameFeaturePluginState GetFailureTransitionState() const = 0;
+
+	/** Determine what kind of release request flags we submit */
+	virtual EInstallBundleReleaseRequestFlags GetReleaseRequestFlags() const = 0;
+};
+
+struct FGameFeaturePluginState_Uninstalling : public FBaseDataReleaseGameFeaturePluginState
+{
+	FGameFeaturePluginState_Uninstalling(FGameFeaturePluginStateMachineProperties& InStateProperties)
+		: FBaseDataReleaseGameFeaturePluginState(InStateProperties)
+	{}
+
+	virtual EGameFeaturePluginState GetSuccessTransitionState() const override
+	{
+		return EGameFeaturePluginState::Terminal;
+	}
+	
+	virtual EGameFeaturePluginState GetFailureTransitionState() const override
+	{
+		return EGameFeaturePluginState::ErrorUninstalling;
+	}
+
+	//Must be overridden to determine what kind of release request we submit
+	virtual EInstallBundleReleaseRequestFlags GetReleaseRequestFlags() const override
+	{
+		return EInstallBundleReleaseRequestFlags::RemoveFilesIfPossible;
+	}
+
+	virtual bool TryUpdateURLData(const FString& NewPluginURL) override
+	{
+		//Use base functionality to update our metadata
+		if (!FGameFeaturePluginState::TryUpdateURLData(NewPluginURL))
+		{
+			return false;
+		}
+
+		//If we are no longer uninstalling before terminate, just exit now as a success immediately
+		if (!ShouldVisitUninstallStateBeforeTerminal())
+		{
+			FBaseDataReleaseGameFeaturePluginState::CleanUp();
+
+			Result = MakeValue();
+			bWasDeleted = true;
+			UpdateStateMachineImmediate();
+
+			return true;
+		}		
+		
+		//Restart our remove request to handle other changes
+		BeginRemoveRequest();
+
+		return true;
+	}
+};
+
+struct FGameFeaturePluginState_Releasing : public FBaseDataReleaseGameFeaturePluginState
+{
+	FGameFeaturePluginState_Releasing(FGameFeaturePluginStateMachineProperties& InStateProperties)
+		: FBaseDataReleaseGameFeaturePluginState(InStateProperties)
+	{}
+
+	virtual EGameFeaturePluginState GetSuccessTransitionState() const override
+	{
+		return EGameFeaturePluginState::StatusKnown;
+	}
+
+	virtual EGameFeaturePluginState GetFailureTransitionState() const override
+	{
+		return EGameFeaturePluginState::ErrorManagingData;
+	}
+
+	//Must be overridden to determine what kind of release request we submit
+	virtual EInstallBundleReleaseRequestFlags GetReleaseRequestFlags() const override
+	{
+		return EInstallBundleReleaseRequestFlags::None;
+	}
 };
 
 struct FGameFeaturePluginState_Downloading : public FGameFeaturePluginState
@@ -584,17 +959,21 @@ struct FGameFeaturePluginState_Downloading : public FGameFeaturePluginState
 			UpdateStateMachineImmediate();
 			return;
 		}
-
-		const TArray<FName>& InstallBundles = StateProperties.ProtocolMetadata.GetSubtype<UE::GameFeatures::FInstallBundlePluginProtocolMetaData>().InstallBundles;
-
-		EInstallBundleRequestFlags InstallFlags = EInstallBundleRequestFlags::None;
+		
+		FInstallBundlePluginProtocolMetaData& Metadata = StateProperties.ProtocolMetadata.GetSubtype<FInstallBundlePluginProtocolMetaData>();
+		const TArray<FName>& InstallBundles = Metadata.InstallBundles;
+		
+		//Pull our InstallFlags from the Metadata, but also make sure SkipMount is set as there is a separate mounting step that will re-request this
+		//without SkipMount and then mount the data, this allows us to pre-download data without mounting it
+		EInstallBundleRequestFlags InstallFlags = Metadata.InstallBundleFlags;
 		InstallFlags |= EInstallBundleRequestFlags::SkipMount;
+
 		TValueOrError<FInstallBundleRequestInfo, EInstallBundleResult> MaybeRequestInfo = BundleManager->RequestUpdateContent(InstallBundles, InstallFlags);
 
-		if (!MaybeRequestInfo.IsValid())
+		if (MaybeRequestInfo.HasError())
 		{
-			ensureMsgf(false, TEXT("Unable to enqueue download for the PluginURL(%s) because %s"), *StateProperties.PluginURL, LexToString(MaybeRequestInfo.GetError()));
-			Result = MakeError(UE::GameFeatures::StateMachineErrorNamespace + TEXT("BundleManager_Cannot_Start_Download"));
+			ensureMsgf(false, TEXT("Unable to enqueue download for the PluginURL(%s) because %s"), *StateProperties.PluginIdentifier.GetFullPluginURL(), LexToString(MaybeRequestInfo.GetError()));
+			Result = GetErrorResult(TEXT("BundleManager.GotState."), LexToString(MaybeRequestInfo.GetError()));
 			UpdateStateMachineImmediate();
 			return;
 		}
@@ -603,8 +982,9 @@ struct FGameFeaturePluginState_Downloading : public FGameFeaturePluginState
 
 		if (EnumHasAnyFlags(RequestInfo.InfoFlags, EInstallBundleRequestInfoFlags::SkippedUnknownBundles))
 		{
-			ensureMsgf(false, TEXT("Unable to enqueue download for the PluginURL(%s) because failed to resolve install bundles!"), *StateProperties.PluginURL);
-			Result = MakeError(UE::GameFeatures::StateMachineErrorNamespace + TEXT("BundleManager_Cannot_Resolve_InstallBundles_For_Download"));
+			ensureMsgf(false, TEXT("Unable to enqueue download for the PluginURL(%s) because failed to resolve install bundles!"), *StateProperties.PluginIdentifier.GetFullPluginURL());
+			Result = GetErrorResult(TEXT("BundleManager.GotState."), TEXT("Resolve_Failed"), UE::GameFeatures::CommonErrorCodes::Generic_ConnectionError);
+			
 			UpdateStateMachineImmediate();
 			return;
 		}
@@ -619,12 +999,20 @@ struct FGameFeaturePluginState_Downloading : public FGameFeaturePluginState
 		{
 			PendingBundleDownloads = MoveTemp(RequestInfo.BundlesEnqueued);
 			IInstallBundleManager::InstallBundleCompleteDelegate.AddRaw(this, &FGameFeaturePluginState_Downloading::OnInstallBundleCompleted);
+			IInstallBundleManager::PausedBundleDelegate.AddRaw(this, &FGameFeaturePluginState_Downloading::OnInstallBundlePaused);
 
 			ProgressTracker = MakeUnique<FInstallBundleCombinedProgressTracker>(false);
 			ProgressTracker->SetBundlesToTrackFromContentState(BundleContentState, PendingBundleDownloads);
 
 			ProgressUpdateHandle = FTSTicker::GetCoreTicker().AddTicker(
 				FTickerDelegate::CreateRaw(this, &FGameFeaturePluginState_Downloading::OnUpdateProgress)/*, 0.1f*/);
+
+			//If this setting is flipped then we should immediately request to pause downloads.
+			//We still generate the downloads so that we have an accurate PendingBundleDownloads list
+			if (Metadata.bUserPauseDownload)
+			{
+				ChangePauseState(true);
+			}
 		}
 	}
 
@@ -639,15 +1027,12 @@ struct FGameFeaturePluginState_Downloading : public FGameFeaturePluginState
 
 		if (!Result.HasError() && BundleResult.Result != EInstallBundleResult::OK)
 		{
-			if (BundleResult.OptionalErrorCode.IsEmpty())
-			{
-				Result = MakeError(UE::GameFeatures::StateMachineErrorNamespace + TEXT("BundleManager_Download_Failure_") + LexToString(BundleResult.Result));
-			}
-			else
-			{
-				Result = MakeError(UE::GameFeatures::StateMachineErrorNamespace + TEXT("BundleManager_Download_Failure_") + BundleResult.OptionalErrorCode);
-			}
-
+			//Use OptionalErrorCode and/or OptionalErrorText if available
+			const FString ErrorCodeEnding = (BundleResult.OptionalErrorCode.IsEmpty()) ? LexToString(BundleResult.Result) : BundleResult.OptionalErrorCode;
+			const FText ErrorText = BundleResult.OptionalErrorCode.IsEmpty() ? UE::GameFeatures::CommonErrorCodes::GetErrorTextForBundleResult(BundleResult.Result) : BundleResult.OptionalErrorText;
+			
+			Result = GetErrorResult(TEXT("BundleManager.OnComplete."), ErrorCodeEnding, ErrorText);
+			
 			if(BundleResult.Result != EInstallBundleResult::UserCancelledError)
 			{
 				TSharedPtr<IInstallBundleManager> BundleManager = IInstallBundleManager::GetPlatformInstallBundleManager();
@@ -679,12 +1064,52 @@ struct FGameFeaturePluginState_Downloading : public FGameFeaturePluginState
 			float Progress = ProgressTracker->GetCurrentCombinedProgress().ProgressPercent;
 			UpdateProgress(Progress);
 
-			//UE_LOG(LogGameFeatures, Display, TEXT("Download Progress: %f for PluginURL(%s)"), Progress, *StateProperties.PluginURL);
+			UE_LOG(LogGameFeatures, VeryVerbose, TEXT("Download Progress: %f for PluginURL(%s)"), Progress, *StateProperties.PluginIdentifier.GetFullPluginURL());
 		}
 
 		return true;
 	}
 
+	void ChangePauseState(bool bPause)
+	{
+		if (PendingBundleDownloads.Num() > 0)
+		{
+			TSharedPtr<IInstallBundleManager> BundleManager = IInstallBundleManager::GetPlatformInstallBundleManager();
+
+			if (bPause)
+			{
+				BundleManager->PauseUpdateContent(PendingBundleDownloads);
+			}
+			else
+			{
+				BundleManager->ResumeUpdateContent(PendingBundleDownloads);
+			}
+			BundleManager->RequestPausedBundleCallback();
+
+			//Use same text we use for InstallBundleManager's UserPaused as this was also a user pause
+			const TCHAR* PauseReason = InstallBundleUtil::GetInstallBundlePauseReason(EInstallBundlePauseFlags::UserPaused);
+
+			NotifyPauseChange(bPause, PauseReason);
+		}
+	}
+
+	void OnInstallBundlePaused(FInstallBundlePauseInfo InPauseBundleInfo)
+	{
+		if (PendingBundleDownloads.Contains(InPauseBundleInfo.BundleName))
+		{
+			const bool bIsPaused = (InPauseBundleInfo.PauseFlags != EInstallBundlePauseFlags::None);
+			const TCHAR* PauseReason = InstallBundleUtil::GetInstallBundlePauseReason(InPauseBundleInfo.PauseFlags);
+			
+			NotifyPauseChange(bIsPaused, PauseReason);
+		}
+	}
+
+	void NotifyPauseChange(bool bIsPaused, FString PauseReason)
+	{
+		FGameFeaturePauseStateChangeContext Context(LexToString(EGameFeaturePluginState::Downloading), PauseReason, bIsPaused);
+		UGameFeaturesSubsystem::Get().OnGameFeaturePauseChange(StateProperties.PluginIdentifier.GetFullPluginURL(), StateProperties.PluginName, Context);
+	}
+
 	virtual void BeginState() override
 	{
 		Cleanup();
@@ -692,7 +1117,7 @@ struct FGameFeaturePluginState_Downloading : public FGameFeaturePluginState
 		check(StateProperties.GetPluginProtocol() == EGameFeaturePluginProtocol::InstallBundle);
 
 		TSharedPtr<IInstallBundleManager> BundleManager = IInstallBundleManager::GetPlatformInstallBundleManager();
-		const TArray<FName>& InstallBundles = StateProperties.ProtocolMetadata.GetSubtype<UE::GameFeatures::FInstallBundlePluginProtocolMetaData>().InstallBundles;
+		const TArray<FName>& InstallBundles = StateProperties.ProtocolMetadata.GetSubtype<FInstallBundlePluginProtocolMetaData>().InstallBundles;
 
 		GotContentStateHandle = BundleManager->GetContentState(InstallBundles, EInstallBundleGetContentStateFlags::None, true, FInstallBundleGetContentStateDelegate::CreateRaw(this, &FGameFeaturePluginState_Downloading::OnGotContentState));
 	}
@@ -701,7 +1126,7 @@ struct FGameFeaturePluginState_Downloading : public FGameFeaturePluginState
 	{
 		if (!Result.HasValue())
 		{
-			StateStatus.SetTransitionError(EGameFeaturePluginState::ErrorInstalling, Result.StealError());
+			StateStatus.SetTransitionError(EGameFeaturePluginState::ErrorManagingData, Result);
 			return;
 		}
 
@@ -722,6 +1147,52 @@ struct FGameFeaturePluginState_Downloading : public FGameFeaturePluginState
 		}
 	}
 
+	virtual bool TryUpdateURLData(const FString& NewPluginURL) override
+	{
+		//if we don't have any in-progress downloads the default behavior is all we need
+		if (PendingBundleDownloads.Num() == 0)
+		{
+			return FGameFeaturePluginState::TryUpdateURLData(NewPluginURL);
+		}
+		
+		//Need to update our BundleFlags for any bundles we are downloading
+		FInstallBundlePluginProtocolMetaData& OldMetaData = StateProperties.ProtocolMetadata.GetSubtype<FInstallBundlePluginProtocolMetaData>();
+		EInstallBundleRequestFlags OldRequestFlags = OldMetaData.InstallBundleFlags;
+		const bool OldUserPausedFlag = OldMetaData.bUserPauseDownload;
+
+		if (!FGameFeaturePluginState::TryUpdateURLData(NewPluginURL))
+		{
+			return false;
+		}
+		FInstallBundlePluginProtocolMetaData& UpdatedMetaData = StateProperties.ProtocolMetadata.GetSubtype<FInstallBundlePluginProtocolMetaData>();
+		
+		//Update our InstallBundleRequestFlags
+		{
+			EInstallBundleRequestFlags UpdatedRequestFlags = OldMetaData.InstallBundleFlags;
+						
+			EInstallBundleRequestFlags AddFlags = (UpdatedRequestFlags & (~OldRequestFlags));
+			EInstallBundleRequestFlags RemoveFlags = ((~UpdatedRequestFlags) & OldRequestFlags);
+
+			if ((AddFlags != EInstallBundleRequestFlags::None) || (RemoveFlags != EInstallBundleRequestFlags::None))
+			{
+				TSharedPtr<IInstallBundleManager> BundleManager = IInstallBundleManager::GetPlatformInstallBundleManager();
+				BundleManager->UpdateContentRequestFlags(PendingBundleDownloads, AddFlags, RemoveFlags);
+			}
+		}
+
+		//Handle pausing or resuming the download if the bUserPauseDownload flag has changed
+		{
+			if (UpdatedMetaData.bUserPauseDownload != OldUserPausedFlag)
+			{
+				ChangePauseState(UpdatedMetaData.bUserPauseDownload);
+			}
+		}
+
+		//Return true at this point even if we don't change anything explicitly for this state as the 
+		//default FGameFeaturePluginState returned true above and thus changed things
+		return true;
+	}
+
 	virtual void EndState() override
 	{
 		Cleanup();
@@ -740,7 +1211,7 @@ struct FGameFeaturePluginState_Installed : public FDestinationGameFeaturePluginS
 		}
 		else if (StateProperties.Destination < EGameFeaturePluginState::Installed)
 		{
-			StateStatus.SetTransition(EGameFeaturePluginState::Uninstalling);
+			StateStatus.SetTransition(EGameFeaturePluginState::Releasing);
 		}
 	}
 };
@@ -801,7 +1272,7 @@ struct FGameFeaturePluginState_Unmounting : public FGameFeaturePluginState
 {
 	FGameFeaturePluginState_Unmounting(FGameFeaturePluginStateMachineProperties& InStateProperties) : FGameFeaturePluginState(InStateProperties) {}
 
-	UE::GameFeatures::FResult Result{ MakeValue() };
+	UE::GameFeatures::FResult Result = MakeValue();
 	TArray<FName> PendingBundles;
 	bool bUnmounted = false;
 
@@ -816,7 +1287,7 @@ struct FGameFeaturePluginState_Unmounting : public FGameFeaturePluginState
 
 		if (!Result.HasError() && BundleResult.Result != EInstallBundleReleaseResult::OK)
 		{
-			Result = MakeError(UE::GameFeatures::StateMachineErrorNamespace + TEXT("BundleManager_Unmount_Error_") + LexToString(BundleResult.Result));
+			Result = GetErrorResult(TEXT("BundleManager.OnReleased."), BundleResult.Result);
 		}
 
 		if (PendingBundles.Num() > 0)
@@ -848,8 +1319,8 @@ struct FGameFeaturePluginState_Unmounting : public FGameFeaturePluginState
 			FText FailureReason;
 			if (!IPluginManager::Get().UnmountExplicitlyLoadedPlugin(StateProperties.PluginName, &FailureReason))
 			{
-				ensureMsgf(false, TEXT("Failed to explicitly unmount the PluginURL(%s) because %s"), *StateProperties.PluginURL, *FailureReason.ToString());
-				Result = MakeError(UE::GameFeatures::StateMachineErrorNamespace + TEXT("Plugin_Cannot_Explicitly_Unmount"));
+				ensureMsgf(false, TEXT("Failed to explicitly unmount the PluginURL(%s) because %s"), *StateProperties.PluginIdentifier.GetFullPluginURL(), *FailureReason.ToString());
+				Result = GetErrorResult(TEXT("Plugin_Cannot_Explicitly_Unmount"));
 				return;
 			}
 		}
@@ -868,15 +1339,15 @@ struct FGameFeaturePluginState_Unmounting : public FGameFeaturePluginState
 
 		TSharedPtr<IInstallBundleManager> BundleManager = IInstallBundleManager::GetPlatformInstallBundleManager();
 
-		const TArray<FName>& InstallBundles = StateProperties.ProtocolMetadata.GetSubtype<UE::GameFeatures::FInstallBundlePluginProtocolMetaData>().InstallBundles;
+		const TArray<FName>& InstallBundles = StateProperties.ProtocolMetadata.GetSubtype<FInstallBundlePluginProtocolMetaData>().InstallBundles;
 
 		EInstallBundleReleaseRequestFlags ReleaseFlags = EInstallBundleReleaseRequestFlags::None;
 		TValueOrError<FInstallBundleRequestInfo, EInstallBundleResult> MaybeRequestInfo = BundleManager->RequestReleaseContent(InstallBundles, ReleaseFlags);
 
-		if (!MaybeRequestInfo.IsValid())
+		if (MaybeRequestInfo.HasError())
 		{
-			ensureMsgf(false, TEXT("Unable to enqueue unmount for the PluginURL(%s) because %s"), *StateProperties.PluginURL, LexToString(MaybeRequestInfo.GetError()));
-			Result = MakeError(UE::GameFeatures::StateMachineErrorNamespace + TEXT("BundleManager_Cannot_Start_Unmount"));
+			ensureMsgf(false, TEXT("Unable to enqueue unmount for the PluginURL(%s) because %s"), *StateProperties.PluginIdentifier.GetFullPluginURL(), LexToString(MaybeRequestInfo.GetError()));
+			Result = GetErrorResult(TEXT("BundleManager.Begin."), MaybeRequestInfo.GetError());
 			return;
 		}
 
@@ -884,8 +1355,8 @@ struct FGameFeaturePluginState_Unmounting : public FGameFeaturePluginState
 
 		if (EnumHasAnyFlags(RequestInfo.InfoFlags, EInstallBundleRequestInfoFlags::SkippedUnknownBundles))
 		{
-			ensureMsgf(false, TEXT("Unable to enqueue unmount for the PluginURL(%s) because failed to resolve install bundles!"), *StateProperties.PluginURL);
-			Result = MakeError(UE::GameFeatures::StateMachineErrorNamespace + TEXT("BundleManager_Cannot_Resolve_InstallBundles_For_Unmount"));
+			ensureMsgf(false, TEXT("Unable to enqueue unmount for the PluginURL(%s) because failed to resolve install bundles!"), *StateProperties.PluginIdentifier.GetFullPluginURL());
+			Result = GetErrorResult(TEXT("BundleManager.Begin."), TEXT("Cannot_Resolve"), UE::GameFeatures::CommonErrorCodes::Generic_ConnectionError);
 			return;
 		}
 
@@ -904,7 +1375,7 @@ struct FGameFeaturePluginState_Unmounting : public FGameFeaturePluginState
 	{
 		if (!Result.HasValue())
 		{
-			StateStatus.SetTransitionError(EGameFeaturePluginState::ErrorMounting, Result.StealError());
+			StateStatus.SetTransitionError(EGameFeaturePluginState::ErrorMounting, Result);
 			return;
 		}
 
@@ -946,11 +1417,11 @@ struct FGameFeaturePluginState_Mounting : public FGameFeaturePluginState
 		{
 			if (BundleResult.OptionalErrorCode.IsEmpty())
 			{
-				Result = MakeError(UE::GameFeatures::StateMachineErrorNamespace + TEXT("BundleManager_Mount_Error_") + LexToString(BundleResult.Result));
+				Result = GetErrorResult(TEXT("BundleManager.OnComplete."), BundleResult.Result);
 			}
 			else
 			{
-				Result = MakeError(UE::GameFeatures::StateMachineErrorNamespace + TEXT("BundleManager_Mount_Error_") + BundleResult.OptionalErrorCode);
+				Result = GetErrorResult(TEXT("BundleManager.OnComplete."), BundleResult.OptionalErrorCode, BundleResult.OptionalErrorText);
 			}
 		}
 
@@ -971,7 +1442,7 @@ struct FGameFeaturePluginState_Mounting : public FGameFeaturePluginState
 	{
 		if (FPakFile* Pak = (FPakFile*)(&PakFile))
 		{
-			UE_LOG(LogGameFeatures, Display, TEXT("Mounted Pak File for (%s) with following files:"), *StateProperties.PluginURL);
+			UE_LOG(LogGameFeatures, Display, TEXT("Mounted Pak File for (%s) with following files:"), *StateProperties.PluginIdentifier.GetFullPluginURL());
 			TArray<FString> OutFileList;
 			Pak->GetPrunedFilenames(OutFileList);
 			for (const FString& FileName : OutFileList)
@@ -995,20 +1466,19 @@ struct FGameFeaturePluginState_Mounting : public FGameFeaturePluginState
 		
 		TSharedPtr<IInstallBundleManager> BundleManager = IInstallBundleManager::GetPlatformInstallBundleManager();
 
-		const TArray<FName>& InstallBundles = StateProperties.ProtocolMetadata.GetSubtype<UE::GameFeatures::FInstallBundlePluginProtocolMetaData>().InstallBundles;
-		
-		// JMarcus TODO: Async Mounting?
-		EInstallBundleRequestFlags InstallFlags = EInstallBundleRequestFlags::None;
+		FInstallBundlePluginProtocolMetaData MetaData = StateProperties.ProtocolMetadata.GetSubtype<FInstallBundlePluginProtocolMetaData>();
+		const TArray<FName>& InstallBundles = MetaData.InstallBundles;
+		EInstallBundleRequestFlags InstallFlags = MetaData.InstallBundleFlags;
 
 		// Make bundle manager use verbose log level for most logs.
 		// We are already done with downloading, so we don't care about logging too much here unless mounting fails.
 		const ELogVerbosity::Type InstallBundleManagerVerbosityOverride = ELogVerbosity::Verbose;
 		TValueOrError<FInstallBundleRequestInfo, EInstallBundleResult> MaybeRequestInfo = BundleManager->RequestUpdateContent(InstallBundles, InstallFlags, InstallBundleManagerVerbosityOverride);
 
-		if (!MaybeRequestInfo.IsValid())
+		if (MaybeRequestInfo.HasError())
 		{
-			ensureMsgf(false, TEXT("Unable to enqueue mount for the PluginURL(%s) because %s"), *StateProperties.PluginURL, LexToString(MaybeRequestInfo.GetError()));
-			Result = MakeError(UE::GameFeatures::StateMachineErrorNamespace + TEXT("BundleManager_Cannot_Start_Mount"));
+			ensureMsgf(false, TEXT("Unable to enqueue mount for the PluginURL(%s) because %s"), *StateProperties.PluginIdentifier.GetFullPluginURL(), LexToString(MaybeRequestInfo.GetError()));
+			Result = GetErrorResult(TEXT("BundleManager.Begin.CannotStart."), MaybeRequestInfo.GetError());
 			return;
 		}
 
@@ -1016,8 +1486,8 @@ struct FGameFeaturePluginState_Mounting : public FGameFeaturePluginState
 
 		if (EnumHasAnyFlags(RequestInfo.InfoFlags, EInstallBundleRequestInfoFlags::SkippedUnknownBundles))
 		{
-			ensureMsgf(false, TEXT("Unable to enqueue mount for the PluginURL(%s) because failed to resolve install bundles!"), *StateProperties.PluginURL);
-			Result = MakeError(UE::GameFeatures::StateMachineErrorNamespace + TEXT("BundleManager_Cannot_Resolve_InstallBundles_For_Mount"));
+			ensureMsgf(false, TEXT("Unable to enqueue mount for the PluginURL(%s) because failed to resolve install bundles!"), *StateProperties.PluginIdentifier.GetFullPluginURL());
+			Result = GetErrorResult(TEXT("BundleManager.Begin."), TEXT("Resolve_Failed"));
 			return;
 		}
 
@@ -1040,7 +1510,7 @@ struct FGameFeaturePluginState_Mounting : public FGameFeaturePluginState
 	{
 		if (!Result.HasValue())
 		{
-			StateStatus.SetTransitionError(EGameFeaturePluginState::ErrorMounting, Result.StealError());
+			StateStatus.SetTransitionError(EGameFeaturePluginState::ErrorMounting, Result);
 			return;
 		}
 		if (!bMounted)
@@ -1048,8 +1518,8 @@ struct FGameFeaturePluginState_Mounting : public FGameFeaturePluginState
 			return;
 		}
 
-		checkf(!StateProperties.PluginInstalledFilename.IsEmpty(), TEXT("PluginInstalledFilename must be set by the Mounting. PluginURL: %s"), *StateProperties.PluginURL);
-		checkf(FPaths::GetExtension(StateProperties.PluginInstalledFilename) == TEXT("uplugin"), TEXT("PluginInstalledFilename must have a uplugin extension. PluginURL: %s"), *StateProperties.PluginURL);
+		checkf(!StateProperties.PluginInstalledFilename.IsEmpty(), TEXT("PluginInstalledFilename must be set by the Mounting. PluginURL: %s"), *StateProperties.PluginIdentifier.GetFullPluginURL());
+		checkf(FPaths::GetExtension(StateProperties.PluginInstalledFilename) == TEXT("uplugin"), TEXT("PluginInstalledFilename must have a uplugin extension. PluginURL: %s"), *StateProperties.PluginIdentifier.GetFullPluginURL());
 
 		// refresh the plugins list to let the plugin manager know about it
 		const TSharedPtr<IPlugin> MaybePlugin = IPluginManager::Get().FindPlugin(StateProperties.PluginName);
@@ -1059,7 +1529,7 @@ struct FGameFeaturePluginState_Mounting : public FGameFeaturePluginState
 		{
 			if (!FPaths::IsSamePath(MaybePlugin->GetDescriptorFileName(), StateProperties.PluginInstalledFilename))
 			{
-				StateStatus.SetTransitionError(EGameFeaturePluginState::ErrorMounting, UE::GameFeatures::StateMachineErrorNamespace + TEXT("Plugin_Name_Already_In_Use"));
+				StateStatus.SetTransitionError(EGameFeaturePluginState::ErrorMounting, GetErrorResult(TEXT("Plugin_Name_Already_In_Use")));
 				return;
 			}
 		}
@@ -1068,7 +1538,7 @@ struct FGameFeaturePluginState_Mounting : public FGameFeaturePluginState
 			const bool bAddedPlugin = IPluginManager::Get().AddToPluginsList(StateProperties.PluginInstalledFilename);
 			if (!bAddedPlugin)
 			{
-				StateStatus.SetTransitionError(EGameFeaturePluginState::ErrorMounting, UE::GameFeatures::StateMachineErrorNamespace + TEXT("Failed_To_Register_Plugin"));
+				StateStatus.SetTransitionError(EGameFeaturePluginState::ErrorMounting, GetErrorResult(TEXT("Failed_To_Register_Plugin")));
 				return;
 			}
 
@@ -1083,21 +1553,16 @@ struct FGameFeaturePluginState_Mounting : public FGameFeaturePluginState
 		// After the new plugin is mounted add the asset registry for that plugin.
 		if (StateProperties.GetPluginProtocol() == EGameFeaturePluginProtocol::InstallBundle)
 		{
+			const FString PluginFolder = FPaths::GetPath(StateProperties.PluginInstalledFilename);
+			const FString PluginAssetRegistry = PluginFolder / TEXT("AssetRegistry.bin");
+
 			TSharedPtr<IPlugin> NewlyMountedPlugin = IPluginManager::Get().FindPlugin(StateProperties.PluginName);
-			if (NewlyMountedPlugin.IsValid() && NewlyMountedPlugin->CanContainContent())
+			if (NewlyMountedPlugin.IsValid() && NewlyMountedPlugin->CanContainContent() && IFileManager::Get().FileExists(*PluginAssetRegistry))
 			{
 				TArray<uint8> SerializedAssetData;
-				const FString PluginFolder = FPaths::GetPath(StateProperties.PluginInstalledFilename);
-				const FString PluginAssetRegistry = PluginFolder / TEXT("AssetRegistry.bin");
-				if (!ensure(IFileManager::Get().FileExists(*PluginAssetRegistry)))
-				{
-					StateStatus.SetTransitionError(EGameFeaturePluginState::ErrorMounting, UE::GameFeatures::StateMachineErrorNamespace + TEXT("Plugin_AssetRegistry_Not_Found"));
-					return;
-				}
-
 				if (!FFileHelper::LoadFileToArray(SerializedAssetData, *PluginAssetRegistry))
 				{
-					StateStatus.SetTransitionError(EGameFeaturePluginState::ErrorMounting, UE::GameFeatures::StateMachineErrorNamespace + TEXT("Failed_To_Load_Plugin_AssetRegistry"));
+					StateStatus.SetTransitionError(EGameFeaturePluginState::ErrorMounting, GetErrorResult(TEXT("Failed_To_Load_Plugin_AssetRegistry")));
 					return;
 				}
 
@@ -1116,6 +1581,7 @@ struct FGameFeaturePluginState_Mounting : public FGameFeaturePluginState
 	virtual void EndState() override
 	{
 		IInstallBundleManager::InstallBundleCompleteDelegate.RemoveAll(this);
+		IInstallBundleManager::PausedBundleDelegate.RemoveAll(this);
 		FCoreDelegates::OnPakFileMounted2.RemoveAll(this);
 	}
 };
@@ -1144,8 +1610,8 @@ struct FGameFeaturePluginState_WaitingForDependencies : public FGameFeaturePlugi
 
 	virtual void UpdateState(FGameFeaturePluginStateStatus& StateStatus) override
 	{
-		checkf(!StateProperties.PluginInstalledFilename.IsEmpty(), TEXT("PluginInstalledFilename must be set by the loading dependencies phase. PluginURL: %s"), *StateProperties.PluginURL);
-		checkf(FPaths::GetExtension(StateProperties.PluginInstalledFilename) == TEXT("uplugin"), TEXT("PluginInstalledFilename must have a uplugin extension. PluginURL: %s"), *StateProperties.PluginURL);
+		checkf(!StateProperties.PluginInstalledFilename.IsEmpty(), TEXT("PluginInstalledFilename must be set by the loading dependencies phase. PluginURL: %s"), *StateProperties.PluginIdentifier.GetFullPluginURL());
+		checkf(FPaths::GetExtension(StateProperties.PluginInstalledFilename) == TEXT("uplugin"), TEXT("PluginInstalledFilename must have a uplugin extension. PluginURL: %s"), *StateProperties.PluginIdentifier.GetFullPluginURL());
 
 		if (!bRequestedDependencies)
 		{
@@ -1155,7 +1621,7 @@ struct FGameFeaturePluginState_WaitingForDependencies : public FGameFeaturePlugi
 			if (!GameFeaturesSubsystem.FindOrCreatePluginDependencyStateMachines(StateProperties.PluginInstalledFilename, Dependencies))
 			{
 				// Failed to query dependencies
-				StateStatus.SetTransitionError(EGameFeaturePluginState::ErrorWaitingForDependencies, UE::GameFeatures::StateMachineErrorNamespace + TEXT("Failed_Dependency_Query"));
+				StateStatus.SetTransitionError(EGameFeaturePluginState::ErrorWaitingForDependencies, GetErrorResult(TEXT("Failed_Dependency_Query")));
 				return;
 			}
 			
@@ -1175,13 +1641,13 @@ struct FGameFeaturePluginState_WaitingForDependencies : public FGameFeaturePlugi
 			if (!RemainingDependency)
 			{
 				// One of the dependency state machines was destroyed before finishing
-				StateStatus.SetTransitionError(EGameFeaturePluginState::ErrorWaitingForDependencies, UE::GameFeatures::StateMachineErrorNamespace + TEXT("Dependency_Destroyed_Before_Finish"));
+				StateStatus.SetTransitionError(EGameFeaturePluginState::ErrorWaitingForDependencies, GetErrorResult(TEXT("Dependency_Destroyed_Before_Finish")));
 				return;
 			}
 
 			if (Pair.Value.HasError())
 			{
-				StateStatus.SetTransitionError(EGameFeaturePluginState::ErrorWaitingForDependencies, UE::GameFeatures::StateMachineErrorNamespace + TEXT("Failed_Dependency_Register"));
+				StateStatus.SetTransitionError(EGameFeaturePluginState::ErrorWaitingForDependencies, GetErrorResult(TEXT("Failed_Dependency_Register")));
 				return;
 			}
 		}
@@ -1208,7 +1674,7 @@ struct FGameFeaturePluginState_WaitingForDependencies : public FGameFeaturePlugi
 				FGameFeatureStateTransitionCanceled::CreateRaw(this, &FGameFeaturePluginState_WaitingForDependencies::OnDependencyTransitionCanceled));
 			if (!ensure(bCancelPending))
 			{
-				OnDependencyTransitionComplete(Dependency, MakeError(TEXT("Failed_Dependency_Register")));
+				OnDependencyTransitionComplete(Dependency, GetErrorResult(UE::GameFeatures::CommonErrorCodes::DependencyFailedRegister));
 			}
 		}
 	}
@@ -1242,7 +1708,7 @@ struct FGameFeaturePluginState_WaitingForDependencies : public FGameFeaturePlugi
 
 		if (!ensure(bSetDestination))
 		{
-			OnDependencyTransitionComplete(Dependency, MakeError(TEXT("Failed_Dependency_Register")));
+			OnDependencyTransitionComplete(Dependency, GetErrorResult(UE::GameFeatures::CommonErrorCodes::DependencyFailedRegister));
 		}
 	}
 
@@ -1314,7 +1780,7 @@ struct FGameFeaturePluginState_Unregistering : public FGameFeaturePluginState
 
 		if (StateProperties.GameFeatureData)
 		{
-			UGameFeaturesSubsystem::Get().OnGameFeatureUnregistering(StateProperties.GameFeatureData, StateProperties.PluginName, StateProperties.PluginURL);
+			UGameFeaturesSubsystem::Get().OnGameFeatureUnregistering(StateProperties.GameFeatureData, StateProperties.PluginName, StateProperties.PluginIdentifier.GetFullPluginURL());
 
 			UGameFeaturesSubsystem::Get().RemoveGameFeatureFromAssetManager(StateProperties.GameFeatureData, StateProperties.PluginName, StateProperties.AddedPrimaryAssetTypes);
 			StateProperties.AddedPrimaryAssetTypes.Empty();
@@ -1396,12 +1862,12 @@ struct FGameFeaturePluginState_Registering : public FGameFeaturePluginState
 			check(StateProperties.AddedPrimaryAssetTypes.Num() == 0);
 			UGameFeaturesSubsystem::Get().AddGameFeatureToAssetManager(StateProperties.GameFeatureData, StateProperties.PluginName, StateProperties.AddedPrimaryAssetTypes);
 
-			UGameFeaturesSubsystem::Get().OnGameFeatureRegistering(StateProperties.GameFeatureData, StateProperties.PluginName, StateProperties.PluginURL);
+			UGameFeaturesSubsystem::Get().OnGameFeatureRegistering(StateProperties.GameFeatureData, StateProperties.PluginName, StateProperties.PluginIdentifier.GetFullPluginURL());
 		}
 		else
 		{
 			// The gamefeaturedata does not exist. The pak file may not be openable or this is a builtin plugin where the pak file does not exist.
-			StateStatus.SetTransitionError(EGameFeaturePluginState::ErrorRegistering, UE::GameFeatures::StateMachineErrorNamespace + TEXT("Plugin_Missing_GameFeatureData"));
+			StateStatus.SetTransitionError(EGameFeaturePluginState::ErrorRegistering, GetErrorResult(TEXT("Plugin_Missing_GameFeatureData")));
 		}
 	}
 };
@@ -1502,7 +1968,7 @@ struct FGameFeaturePluginState_Loading : public FGameFeaturePluginState
 			BundleHandle->WaitUntilComplete(0.0f, false);
 		}
 
-		UGameFeaturesSubsystem::Get().OnGameFeatureLoading(StateProperties.GameFeatureData, StateProperties.PluginURL);
+		UGameFeaturesSubsystem::Get().OnGameFeatureLoading(StateProperties.GameFeatureData, StateProperties.PluginIdentifier.GetFullPluginURL());
 
 		StateStatus.SetTransition(EGameFeaturePluginState::Loaded);
 	}
@@ -1596,12 +2062,26 @@ struct FGameFeaturePluginState_Deactivating : public FGameFeaturePluginState
 
 			// Deactivate
 			FGameFeatureDeactivatingContext Context(FSimpleDelegate::CreateRaw(this, &FGameFeaturePluginState_Deactivating::OnPauserCompleted));
-			UGameFeaturesSubsystem::Get().OnGameFeatureDeactivating(StateProperties.GameFeatureData, StateProperties.PluginName, Context, StateProperties.PluginURL);
+			UGameFeaturesSubsystem::Get().OnGameFeatureDeactivating(StateProperties.GameFeatureData, StateProperties.PluginName, Context, StateProperties.PluginIdentifier.GetFullPluginURL());
 			NumExpectedPausers = Context.NumPausers;
+
+			// Since we are pausing work during this deactivation, also notify the OnGameFeaturePauseChange delegate
+			if (NumExpectedPausers > 0)
+			{
+				FGameFeaturePauseStateChangeContext PauseContext(LexToString(EGameFeaturePluginState::Deactivating), TEXT("PendingDeactivationCallbacks"), true);
+				UGameFeaturesSubsystem::Get().OnGameFeaturePauseChange(StateProperties.PluginIdentifier.GetFullPluginURL(), StateProperties.PluginName, PauseContext);
+			}
 		}
 
 		if (NumExpectedPausers == NumObservedPausers)
 		{
+			//If we previously sent an OnGameFeaturePauseChange delegate we need to send that work is now unpaused
+			if (NumExpectedPausers > 0)
+			{
+				FGameFeaturePauseStateChangeContext PauseContext(LexToString(EGameFeaturePluginState::Deactivating), TEXT(""), false);
+				UGameFeaturesSubsystem::Get().OnGameFeaturePauseChange(StateProperties.PluginIdentifier.GetFullPluginURL(), StateProperties.PluginName, PauseContext);
+			}
+
 			if (!bRequestedGC && StateProperties.Destination.MaxState == EGameFeaturePluginState::Loaded)
 			{
 				// If we aren't going farther than Loaded, GC now
@@ -1634,7 +2114,7 @@ struct FGameFeaturePluginState_Activating : public FGameFeaturePluginState
 
 		StateProperties.GameFeatureData->InitializeHierarchicalPluginIniFiles(StateProperties.PluginInstalledFilename);
 
-		UGameFeaturesSubsystem::Get().OnGameFeatureActivating(StateProperties.GameFeatureData, StateProperties.PluginName, Context, StateProperties.PluginURL);
+		UGameFeaturesSubsystem::Get().OnGameFeatureActivating(StateProperties.GameFeatureData, StateProperties.PluginName, Context, StateProperties.PluginIdentifier.GetFullPluginURL());
 
 		StateStatus.SetTransition(EGameFeaturePluginState::Active);
 	}
@@ -1819,6 +2299,19 @@ bool UGameFeaturePluginStateMachine::TryCancel(FGameFeatureStateTransitionCancel
 	return true;
 }
 
+bool UGameFeaturePluginStateMachine::TryUpdatePluginURLData(const FString& InPluginURL)
+{
+	//Check if our current PluginURL information matches this new URL
+	//and early out if no update is needed
+	if (StateProperties.PluginIdentifier.ExactMatchesURL(InPluginURL))
+	{
+		return false;
+	}
+
+	const EGameFeaturePluginState CurrentState = GetCurrentState();
+	return AllStates[CurrentState]->TryUpdateURLData(InPluginURL);
+}
+
 void UGameFeaturePluginStateMachine::RemovePendingTransitionCallback(FDelegateHandle InHandle)
 {
 	for (std::underlying_type<EGameFeaturePluginState>::type iState = 0;
@@ -1870,13 +2363,13 @@ const FString& UGameFeaturePluginStateMachine::GetGameFeatureName() const
 	}
 	else
 	{
-		return StateProperties.PluginURL;
+		return StateProperties.PluginIdentifier.GetFullPluginURL();
 	}
 }
 
 const FString& UGameFeaturePluginStateMachine::GetPluginURL() const
 {
-	return StateProperties.PluginURL;
+	return StateProperties.PluginIdentifier.GetFullPluginURL();
 }
 
 const FString& UGameFeaturePluginStateMachine::GetPluginName() const
@@ -1912,7 +2405,10 @@ bool UGameFeaturePluginStateMachine::IsRunning() const
 
 bool UGameFeaturePluginStateMachine::IsStatusKnown() const
 {
-	return GetCurrentState() == EGameFeaturePluginState::ErrorUnavailable || GetCurrentState() >= EGameFeaturePluginState::StatusKnown;
+	return GetCurrentState() == EGameFeaturePluginState::ErrorUnavailable || 
+		   GetCurrentState() == EGameFeaturePluginState::Uninstalling || 
+		   GetCurrentState() == EGameFeaturePluginState::ErrorUninstalling ||
+		   GetCurrentState() >= EGameFeaturePluginState::StatusKnown;
 }
 
 bool UGameFeaturePluginStateMachine::IsAvailable() const
@@ -2094,7 +2590,7 @@ FGameFeaturePluginStateMachineProperties::FGameFeaturePluginStateMachineProperti
 	const FGameFeaturePluginStateRange& DesiredDestination,
 	const FGameFeaturePluginRequestUpdateStateMachine& RequestUpdateStateMachineDelegate,
 	const FGameFeatureStateProgressUpdate& FeatureStateProgressUpdateDelegate)
-	: PluginURL(InPluginURL)
+	: PluginIdentifier(InPluginURL)
 	, Destination(DesiredDestination)
 	, OnRequestUpdateStateMachine(RequestUpdateStateMachineDelegate)
 	, OnFeatureStateProgressUpdate(FeatureStateProgressUpdateDelegate)
@@ -2103,71 +2599,265 @@ FGameFeaturePluginStateMachineProperties::FGameFeaturePluginStateMachineProperti
 
 EGameFeaturePluginProtocol FGameFeaturePluginStateMachineProperties::GetPluginProtocol() const
 {
-	if (CachedPluginProtocol != EGameFeaturePluginProtocol::Unknown)
+	if (PluginIdentifier.PluginProtocol != EGameFeaturePluginProtocol::Unknown)
 	{
-		return CachedPluginProtocol;
+		return PluginIdentifier.PluginProtocol;
 	}
-
-	for (EGameFeaturePluginProtocol Proto : TEnumRange<EGameFeaturePluginProtocol>())
+	else
 	{
-		const TCHAR* Prefix = UE::GameFeatures::GameFeaturePluginProtocolPrefix(Proto);
-		if (Prefix && *Prefix && PluginURL.StartsWith(Prefix))
-		{
-			CachedPluginProtocol = Proto;
-			break;
-		}
+		ensureAlwaysMsgf(false, TEXT("No cached valid EGameFeaturePluginProtocol in our PluginIdentifier!"));
+		return UGameFeaturesSubsystem::GetPluginURLProtocol(PluginIdentifier.GetFullPluginURL());
 	}
+}
 
-	return CachedPluginProtocol;
+FInstallBundlePluginProtocolMetaData::FInstallBundlePluginProtocolMetaData()
+{
+	ResetToDefaults();
 }
 
-bool FGameFeaturePluginStateMachineProperties::ParseURL()
+void FInstallBundlePluginProtocolMetaData::ResetToDefaults()
 {
-	if (GetPluginProtocol() == EGameFeaturePluginProtocol::File)
+	InstallBundles.Reset();
+	VersionNum = FDefaultValues::CurrentVersionNum;
+	bUninstallBeforeTerminate = FDefaultValues::Default_bUninstallBeforeTerminate;
+	bUserPauseDownload = FDefaultValues::Default_bUserPauseDownload;
+	InstallBundleFlags = FDefaultValues::Default_InstallBundleFlags;
+}
+
+FString FInstallBundlePluginProtocolMetaData::ToString() const
+{
+	FString ReturnedString;
+
+	//Always encode InstallBundles
+	ReturnedString = LexToString(EGameFeatureInstallBundleProtocolOptions::Bundles) + UE::GameFeatures::PluginURLStructureInfo::OptionAssignOperator;
+	for (const FName& BundleName : InstallBundles)
 	{
-		PluginInstalledFilename = PluginURL.RightChop(
-			FCString::Strlen(UE::GameFeatures::GameFeaturePluginProtocolPrefix(EGameFeaturePluginProtocol::File)));
+		ReturnedString.Append(BundleName.ToString());
 	}
-	else if (GetPluginProtocol() == EGameFeaturePluginProtocol::InstallBundle)
+
+	//Always encode version
+	ReturnedString.Append(FString::Printf(TEXT("%s%s%s%d"), UE::GameFeatures::PluginURLStructureInfo::OptionSeperator, *LexToString(EGameFeatureInstallBundleProtocolOptions::Version), UE::GameFeatures::PluginURLStructureInfo::OptionAssignOperator, VersionNum));
+
+	//Encode InstallBundleFlags if not the default
+	if (InstallBundleFlags != FDefaultValues::Default_InstallBundleFlags)
 	{
-		int32 CursorIdx = FCString::Strlen(UE::GameFeatures::GameFeaturePluginProtocolPrefix(EGameFeaturePluginProtocol::InstallBundle));
-		int32 QueryIdx = PluginURL.Find(TEXT("?"), ESearchCase::IgnoreCase, ESearchDir::FromStart, CursorIdx);
-		if (QueryIdx == INDEX_NONE)
-		{
-			return false;
-		}
-		
-		PluginInstalledFilename = PluginURL.Mid(CursorIdx, QueryIdx - CursorIdx);
-		CursorIdx = QueryIdx + 1;
-		
-		FString BundleNamesString = PluginURL.Mid(CursorIdx);
-		TArray<FString> BundleNames;
-		BundleNamesString.ParseIntoArray(BundleNames, TEXT(","));
-		if (BundleNames.Num() == 0)
-		{
-			return false;
-		}
+		ReturnedString.Append(FString::Printf(TEXT("%s%s%s%d"), UE::GameFeatures::PluginURLStructureInfo::OptionSeperator, *LexToString(EGameFeatureInstallBundleProtocolOptions::Flags), UE::GameFeatures::PluginURLStructureInfo::OptionAssignOperator, static_cast<uint32>(InstallBundleFlags)));
+	}
+
+	//Encode bUninstallBeforeTerminate if not the default
+	if (bUninstallBeforeTerminate != FDefaultValues::Default_bUninstallBeforeTerminate)
+	{
+		ReturnedString.Append(FString::Printf(TEXT("%s%s%s%s"), UE::GameFeatures::PluginURLStructureInfo::OptionSeperator, *LexToString(EGameFeatureInstallBundleProtocolOptions::UninstallBeforeTerminate), UE::GameFeatures::PluginURLStructureInfo::OptionAssignOperator, *LexToString(bUninstallBeforeTerminate)));
+	}
+
+	//Encode bUserPauseDownload if not the default
+	if (bUserPauseDownload != FDefaultValues::Default_bUserPauseDownload)
+	{
+		ReturnedString.Append(FString::Printf(TEXT("%s%s%s%s"), UE::GameFeatures::PluginURLStructureInfo::OptionSeperator, *LexToString(EGameFeatureInstallBundleProtocolOptions::UserPauseDownload), UE::GameFeatures::PluginURLStructureInfo::OptionAssignOperator, *LexToString(bUserPauseDownload)));
+	}
+
+	static_assert(static_cast<uint8>(EGameFeatureInstallBundleProtocolOptions::Count) == 5, "Update this function to handle the newly added EGameFeatureInstallBundleProtocolOptions value!");
+
+	return ReturnedString;
+}
+
+bool FInstallBundlePluginProtocolMetaData::FromString(const FString& URLString, FInstallBundlePluginProtocolMetaData& Metadata)
+{
+	bool bParseSuccess = true;
+	TArray<FString> URLOptions;
+	URLString.ParseIntoArray(URLOptions, UE::GameFeatures::PluginURLStructureInfo::OptionSeperator);
+
+	Metadata.ResetToDefaults();
+
+	if (URLOptions.Num() > 0)
+	{
+		ensureAlwaysMsgf(URLOptions[0].StartsWith(UE::GameFeatures::GameFeaturePluginProtocolPrefix(EGameFeaturePluginProtocol::InstallBundle)), TEXT("Unexpected URL Format! Expected Protocol and uplugin information at the beginning of the URL %s"), *URLString);
 
-		UE::GameFeatures::FInstallBundlePluginProtocolMetaData& MetaData = *ProtocolMetadata.SetSubtype<UE::GameFeatures::FInstallBundlePluginProtocolMetaData>();
-		MetaData.InstallBundles.Reserve(BundleNames.Num());
-		for (FString& BundleNameString : BundleNames)
+		//Parse through our URLOptions. Start at ParsingIndex 1 as option 0 should always be the .uplugin path that doesn't contain metadata
+		for (int ParsingIndex = 1; ParsingIndex < URLOptions.Num(); ++ParsingIndex)
 		{
-			MetaData.InstallBundles.Add(*BundleNameString);
+			const FString& URLOption = URLOptions[ParsingIndex];
+	
+			TArray<FString> OptionStrings;
+			URLOption.ParseIntoArray(OptionStrings, UE::GameFeatures::PluginURLStructureInfo::OptionAssignOperator);
+
+			if (OptionStrings.Num() > 0)
+			{
+				EGameFeatureInstallBundleProtocolOptions OptionEnum = LexFromString(OptionStrings[0]);
+
+				switch (OptionEnum)
+				{
+					case EGameFeatureInstallBundleProtocolOptions::Bundles:
+					{
+						if (OptionStrings.Num() != 2)
+						{
+							bParseSuccess = false;
+							UE_LOG(LogGameFeatures, Error, TEXT("Error parsing InstallBundle protocol options URL %s . No Valid Bundle List Found!"), *URLString);
+						}
+						else
+						{
+							TArray<FString> BundleNames;
+							OptionStrings[1].ParseIntoArray(BundleNames, TEXT(","));
+
+							Metadata.InstallBundles.Reserve(BundleNames.Num());
+							for (FString& BundleNameString : BundleNames)
+		{
+								Metadata.InstallBundles.Add(*BundleNameString);
+							}
+						}
+
+						break;
+					}
+
+					case EGameFeatureInstallBundleProtocolOptions::Version:
+					{
+						if (OptionStrings.Num() != 2)
+						{
+							bParseSuccess = false;
+							UE_LOG(LogGameFeatures, Warning, TEXT("Error parsing InstallBundle protocol options URL %s . Invalid Version Option!"), *URLString);
+						}
+						else
+						{
+							Metadata.VersionNum = FCString::Strtoi(*OptionStrings[1], nullptr, 10);
+							if (Metadata.VersionNum != FDefaultValues::CurrentVersionNum)
+							{
+								ensureAlwaysMsgf(false, TEXT("Mismatched version number for URL %s ! URL was generated on a previous version and might be missing data!"));
+							}
+						}
+
+						break;
+					}
+
+					case EGameFeatureInstallBundleProtocolOptions::Flags:
+					{
+						if (OptionStrings.Num() != 2)
+						{
+							bParseSuccess = false;
+							UE_LOG(LogGameFeatures, Warning, TEXT("Error parsing InstallBundle protocol options URL %s . Invalid Flags Option!"), *URLString);
+						}
+						//Expect to parse uint32 flag version of this as a 2nd parameter
+						else
+						{
+							Metadata.InstallBundleFlags = static_cast<EInstallBundleRequestFlags>(FCString::Strtoi(*OptionStrings[1], nullptr, 10));
+						}
+
+						break;
+					}
+
+					case EGameFeatureInstallBundleProtocolOptions::UninstallBeforeTerminate:
+					{
+						if (OptionStrings.Num() != 2)
+						{
+							bParseSuccess = false;
+							UE_LOG(LogGameFeatures, Warning, TEXT("Error parsing InstallBundle protocol options URL %s . Invalid UninstallBeforeTerminate Option!"), *URLString);
+						}
+						else
+						{
+							Metadata.bUninstallBeforeTerminate = OptionStrings[1].ToBool();
+						}
+
+						break;
+					}
+
+					case EGameFeatureInstallBundleProtocolOptions::UserPauseDownload:
+					{
+						if (OptionStrings.Num() != 2)
+						{
+							bParseSuccess = false;
+							UE_LOG(LogGameFeatures, Warning, TEXT("Error parsing InstallBundle protocol options URL %s . Invalid UserPauseDownload Option!"), *URLString);
+						}
+						else
+						{
+							Metadata.bUserPauseDownload = OptionStrings[1].ToBool();
+						}
+
+						break;
+					}
+
+					case EGameFeatureInstallBundleProtocolOptions::Count:
+					default:
+					{
+						bParseSuccess = false;
+						UE_LOG(LogGameFeatures, Error, TEXT("Error parsing InstallBundle protocol options for URL %s . Unknown Option %s"), *URLString, *URLOption);
+
+						break;
+					}
+				}
+			}
 		}
 	}
-	else
+	//We require to have InstallBundle names for this URL parse to be correct
+	if (Metadata.InstallBundles.Num() == 0)
+	{
+		bParseSuccess = false;
+		UE_LOG(LogGameFeatures, Error, TEXT("Error parsing InstallBundle protocol options URL %s . No Bundle List Found!"), *URLString);
+	}
+
+	static_assert(static_cast<uint8>(EGameFeatureInstallBundleProtocolOptions::Count) == 5, "Update this function to handle the newly added EGameFeatureInstallBundleProtocolOptions value!");
+
+	return bParseSuccess;
+}
+
+bool FGameFeaturePluginStateMachineProperties::ParseURL()
+{
+	if (!ensureMsgf(!PluginIdentifier.IdentifyingURLSubset.IsEmpty(), TEXT("Unexpected empty IdentifyingURLSubset while parsing URL!")))
 	{
-		ensureMsgf(false, TEXT("Unknown protocol for PluginURL: %s"), *PluginURL);
 		return false;
 	}
 
+	//The current IdentifyingURLSubset IS the same as the PluginInstalledFilename, so just use the already parsed version
+	PluginInstalledFilename = PluginIdentifier.IdentifyingURLSubset;
 	PluginName = FPaths::GetBaseFilename(PluginInstalledFilename);
 
 	if (PluginInstalledFilename.IsEmpty() || !PluginInstalledFilename.EndsWith(TEXT(".uplugin")))
 	{
-		ensureMsgf(false, TEXT("PluginInstalledFilename must have a uplugin extension. PluginURL: %s"), *PluginURL);
+		ensureMsgf(false, TEXT("PluginInstalledFilename must have a uplugin extension. PluginInstalledFilename: %s"), *PluginInstalledFilename);
 		return false;
 	}
 
+	//Do additional parsing of our Metadata from the options on our remaining URL
+	if (GetPluginProtocol() == EGameFeaturePluginProtocol::InstallBundle)
+	{
+		FInstallBundlePluginProtocolMetaData& MetaData = *ProtocolMetadata.SetSubtype<FInstallBundlePluginProtocolMetaData>();
+		if (!FInstallBundlePluginProtocolMetaData::FromString(PluginIdentifier.GetFullPluginURL(), MetaData))
+		{
+			ensureMsgf(false, TEXT("Failure to parse URL %s into a valid FInstallBundlePluginProtocolMetaData"), *PluginIdentifier.GetFullPluginURL());
+			return false;
+		}
+	}
+
+	static_assert(static_cast<uint8>(EGameFeaturePluginProtocol::Count) == 3, "Update FGameFeaturePluginStateMachineProperties::ParseURL to handle any new Metadata parsing required for new EGameFeaturePluginProtocol. If no metadata is required just increment this counter.");
+
 	return true;
 }
+
+bool FGameFeaturePluginStateMachineProperties::ValidateURLUpdate(const FGameFeaturePluginStateMachineProperties& OldProperties) const
+{
+	bool bIsValidUpdate = false;
+
+	//Should never change our Identifier or PluginProtocol
+	bIsValidUpdate = ((PluginIdentifier == OldProperties.PluginIdentifier) &&
+						(GetPluginProtocol() == OldProperties.GetPluginProtocol()));
+
+	if (bIsValidUpdate && (GetPluginProtocol() == EGameFeaturePluginProtocol::InstallBundle))
+	{
+		if (!ensureAlwaysMsgf((ProtocolMetadata.HasSubtype<FInstallBundlePluginProtocolMetaData>() && 
+								OldProperties.ProtocolMetadata.HasSubtype<FInstallBundlePluginProtocolMetaData>())
+								,TEXT("Error with InstallBundle protocol FGameFeaturePluginStateMachineProperties having an invalid ProtocolMetadata. URL: %s")
+								,*PluginIdentifier.GetFullPluginURL()))
+		{
+			
+			return false;
+		}
+
+		const FInstallBundlePluginProtocolMetaData& NewMetaData = ProtocolMetadata.GetSubtype<FInstallBundlePluginProtocolMetaData>();
+		const FInstallBundlePluginProtocolMetaData& OldMetaData = OldProperties.ProtocolMetadata.GetSubtype<FInstallBundlePluginProtocolMetaData>();
+
+		if (!ensureMsgf((NewMetaData.InstallBundles == OldMetaData.InstallBundles), TEXT("Unexpected change in InstallBundles when updating URL: %s to %s"), *PluginIdentifier.GetFullPluginURL(), *OldProperties.PluginIdentifier.GetFullPluginURL()))
+		{
+			bIsValidUpdate = false;
+		}
+	}
+
+	return bIsValidUpdate;
+}
diff --git a/Engine/Plugins/Experimental/GameFeatures/Source/GameFeatures/Private/GameFeaturePluginStateMachine.h b/Engine/Plugins/Experimental/GameFeatures/Source/GameFeatures/Private/GameFeaturePluginStateMachine.h
--- a/Engine/Plugins/Experimental/GameFeatures/Source/GameFeatures/Private/GameFeaturePluginStateMachine.h
+++ b/Engine/Plugins/Experimental/GameFeatures/Source/GameFeatures/Private/GameFeaturePluginStateMachine.h
@@ -8,6 +8,7 @@
 #include "Containers/Ticker.h"
 #include "UObject/ObjectMacros.h"
 #include "UObject/Object.h"
+#include "GameFeaturesSubsystem.h"
 #include "GameFeaturePluginOperationResult.h"
 #include "GameFeatureTypes.h"
 #include "GameFeaturePluginStateMachine.generated.h"
@@ -16,6 +17,8 @@ class UGameFeatureData;
 class UGameFrameworkComponentManager;
 class UGameFeaturePluginStateMachine;
 struct FComponentRequestHandle;
+enum class EInstallBundleResult : int;
+enum class EInstallBundleReleaseResult;
 
 /*
 *************** GameFeaturePlugin state machine graph ***************
@@ -31,45 +34,58 @@ Transition states are expected to transition the machine to another state after
                          +------+-------+
      +------------+             |
      |     *      |             |
-     |  Terminal  <-------------~------------------------------
-     |            |             |                             |
-     +--^------^--+             |                             |
-        |      |                |                             |
-        |      |         +------v-------+                     |
-        |      |         |      *       |                     |
-        |      ----------+UnknownStatus |                     |
-        |                |              |                     |
-        |                +------+-------+                     |
-        |                  |                                  |
-        |      +-----------v---+     +--------------------+   |
-        |      |               |     |         !          |   |
-        |      |CheckingStatus <-----> ErrorCheckingStatus+-->|
-        |      |               |     |                    |   |
-        |      +------+------^-+     +--------------------+   |
-        |             |      |                                |
-        |             |      |       +--------------------+   |
-        ----------    |      |       |         !          |   |
-                 |    |      --------> ErrorUnavailable   +----
-                 |    |              |                    |
-                 |    |              +--------------------+
+     |  Terminal  <-------------~-----------------------------------------------
+     |            |             |                                              |
+     +--^------^--+             ----------------------------                   |
+        |      |                                           |                   |
+        |      |                                    +------v--------+          |
+        |      |                                    |      *        |          |
+        |      -------------------------------------+ UnknownStatus |          |
+        |           ^                      ^        |               |          |
+        |           |                      |        +-------+-------+          |
+        |           |                      |                |                  |
+        |    +------+-------+    *---------+---------+      |                  |
+        |    |              |    |         !         |      |                  |
+        |    | Uninstalling <----> ErrorUninstalling |      |                  |
+        |    |              |    |                   |      |                  |
+        |    +---^----------+    +---------+---------+      |                  |
+        |        |                         |                |                  |
+        |        |    ----------------------                |                  |
+        |        |    |                                     |                  |
+        |        |    |                     -----------------                  |
+        |        |    |                     |                                  |
+        |        |    |         +-----------v---+     +--------------------+   |
+        |        |    |         |               |     |         !          |   |
+        |        |    |         |CheckingStatus <-----> ErrorCheckingStatus+-->|
+        |        |    |         |               |     |                    |   |
+        |        |    |         +------+------^-+     +--------------------+   |
+        |        |    |                |      |                                |
+        |        |    |                |      |       +--------------------+   |
+        ---------~    |                |      |       |         !          |   |
+                 |    |<----------------      --------> ErrorUnavailable   +----
+                 |    |                               |                    |
+                 |    |                               +--------------------+
                  |    |
-               +-+----v-------+                            
-               |      *       |
-     ----------+ StatusKnown  |
-     |         |              |
-     |         +-^-----+---+--+
-     |           |         |
-     |     ------~---------~-------------------------
-     |     |     |         |                        |
-     |  +--v-----+---+   +-v----------+       +-----v--------------+
-     |  |            |   |            |       |         !          |
-     |  |Uninstalling|   |Downloading <-------> ErrorInstalling    |
-     |  |            |   |            |       |                    |
-     |  +--------^---+   +-+----------+       +--------------------+
-     |           |         |
-     |         +-+---------v-+
-     |         |      *      |
-     ---------->  Installed  |
+            +----+----v----+                            
+            |      *       |
+         ---> StatusKnown  +----------------------------------------------
+         |  |              |                                 |           |
+         |  +----------^---+                                 |           |
+         |                                                   |           |
+         |                                                   |           |
+         |                                                   |           |
+         |                                                   |           |
+         |                                                   |           |
+      +--+---------+      +-------------------+       +------v-------+   |
+      |            |      |         !         |       |              |   |
+      | Releasing  <------> ErrorManagingData <-------> Downloading  |   |
+      |            |      |                   |       |              |   |  
+      +--^---------+      +-------------------+       +-------+------+   |  
+         |                                                   |           | 
+         |                                                   |           | 
+         |     +-------------+                               |           |   
+         |     |      *      |                               v           |    
+         ------+ Installed   <--------------------------------------------
                |             |
                +-^---------+-+
                  |         |
@@ -131,11 +147,6 @@ Transition states are expected to transition the machine to another state after
 namespace UE::GameFeatures
 {
 	const TCHAR* GameFeaturePluginProtocolPrefix(EGameFeaturePluginProtocol Protocol);
-
-	struct FInstallBundlePluginProtocolMetaData
-	{
-		TArray<FName> InstallBundles;
-	};
 }
 
 struct FGameFeaturePluginStateRange
@@ -223,12 +234,12 @@ struct FGameFeaturePluginStateMachineProperties
 {
 	GENERATED_BODY()
 
-	/**
-	 * The URL to find the plugin. This takes the form protocol:identifier.
-	 * Every protocol will have its own style of identifier.
-	 * For example, if the file is simply on disk, you can use file:../../../YourGameModule/Plugins/MyPlugin/MyPlugin.uplugin
-	 */
-	FString PluginURL;
+	/** 
+	* The Identifier used to find this Plugin. Parsed from the supplied PluginURL at creation.
+	* Every protocol will have its own style of identifier URL that will get parsed to generate this.
+	* For example, if the file is simply on disk, you can use file:../../../YourGameModule/Plugins/MyPlugin/MyPlugin.uplugin
+	**/
+	FGameFeaturePluginIdentifier PluginIdentifier;
 
 	/** Filename on disk of the .uplugin file. */
 	FString PluginInstalledFilename;
@@ -237,7 +248,7 @@ struct FGameFeaturePluginStateMachineProperties
 	FString PluginName;
 
 	/** Meta data parsed from the URL for a specific protocol. */
-	TUnion<UE::GameFeatures::FInstallBundlePluginProtocolMetaData> ProtocolMetadata;
+	TUnion<FInstallBundlePluginProtocolMetaData> ProtocolMetadata;
 
 	TArray<FName> AddedPrimaryAssetTypes;
 
@@ -275,8 +286,8 @@ struct FGameFeaturePluginStateMachineProperties
 
 	bool ParseURL();
 
-private:
-	mutable EGameFeaturePluginProtocol CachedPluginProtocol = EGameFeaturePluginProtocol::Unknown;
+	/** Checks to see if any invalid data was changed during a URL update. True if data updated was all values expected to be changed. */
+	bool ValidateURLUpdate(const FGameFeaturePluginStateMachineProperties& OldProperties) const;
 };
 
 /** Input and output information for a state's UpdateState */
@@ -287,22 +298,13 @@ private:
 	EGameFeaturePluginState TransitionToState = EGameFeaturePluginState::Uninitialized;
 
 	/** Holds the current error for any state transition. */
-	UE::GameFeatures::FResult TransitionResult{ MakeValue() };
+	UE::GameFeatures::FResult TransitionResult = MakeValue();
 
 	friend class UGameFeaturePluginStateMachine;
 
 public:
-	void SetTransition(EGameFeaturePluginState InTransitionToState)
-	{
-		TransitionToState = InTransitionToState;
-		TransitionResult = MakeValue();
-	}
-
-	void SetTransitionError(EGameFeaturePluginState TransitionToErrorState, FString Error)
-	{
-		TransitionToState = TransitionToErrorState;
-		TransitionResult = MakeError(MoveTemp(Error));
-	}
+	void SetTransition(EGameFeaturePluginState InTransitionToState);
+	void SetTransitionError(EGameFeaturePluginState TransitionToErrorState, UE::GameFeatures::FResult TransitionResult);
 };
 
 enum class EGameFeaturePluginStateType : uint8
@@ -330,6 +332,13 @@ struct FGameFeaturePluginState
 	/** Attempt to cancel any pending state transition. */
 	virtual void TryCancelState() {}
 
+	/** Called if we have updated the URL for this FGameFeaturePluginState.
+		This can be done whenever URL data is updated that isn't tied to our 
+		FGameFeaturePluginIdentifier information. EX: MetaData options information
+		that is parsed from the URL 
+		Returns false if no update occured or the update failed. True on successful update. */
+	virtual bool TryUpdateURLData(const FString& NewPluginURL);
+	
 	/** Called when this state is no longer the active state */
 	virtual void EndState() {}
 
@@ -348,6 +357,18 @@ struct FGameFeaturePluginState
 
 	void UpdateProgress(float Progress) const;
 
+protected:
+	/** Builds an end FResult with some minimal error information with overrides for common types we 
+		need to generate errors from */
+	UE::GameFeatures::FResult GetErrorResult(const FString& ErrorCode, const FText OptionalErrorText = FText()) const;
+	UE::GameFeatures::FResult GetErrorResult(const FString& ErrorNamespaceAddition, const FString& ErrorCode, const FText OptionalErrorText = FText()) const;
+	UE::GameFeatures::FResult GetErrorResult(const FString& ErrorNamespaceAddition, const EInstallBundleResult ErrorResult) const;
+	UE::GameFeatures::FResult GetErrorResult(const FString& ErrorNamespaceAddition, const EInstallBundleReleaseResult ErrorResult) const;
+
+	/** Returns true if this state should transition to the Uninstalled state. 
+	    returns False if it should just go directly to the Terminal state instead. */
+	bool ShouldVisitUninstallStateBeforeTerminal() const;
+
 private:
 	void CleanupDeferredUpdateCallbacks() const;
 
@@ -409,6 +430,10 @@ public:
 	/** Cancel the current transition if possible */
 	bool TryCancel(FGameFeatureStateTransitionCanceled OnFeatureStateTransitionCanceled, FDelegateHandle* OutCallbackHandle = nullptr);
 
+	/** Update the current PluginURL data for this plugin if possible. Returns false if this update fails or
+		if the supplied InPluginURL matches the existing URL data for this plugin.**/
+	bool TryUpdatePluginURLData(const FString& InPluginURL);
+
 	/** Remove any pending callback from SetDestination */
 	void RemovePendingTransitionCallback(FDelegateHandle InHandle);
 
diff --git a/Engine/Plugins/Experimental/GameFeatures/Source/GameFeatures/Private/GameFeaturesSubsystem.cpp b/Engine/Plugins/Experimental/GameFeatures/Source/GameFeatures/Private/GameFeaturesSubsystem.cpp
--- a/Engine/Plugins/Experimental/GameFeatures/Source/GameFeatures/Private/GameFeaturesSubsystem.cpp
+++ b/Engine/Plugins/Experimental/GameFeatures/Source/GameFeatures/Private/GameFeaturesSubsystem.cpp
@@ -17,9 +17,131 @@
 #include "Engine/StreamableManager.h"
 #include "Engine/AssetManagerSettings.h"
 #include "Engine/Engine.h"
+#include "InstallBundleTypes.h"
 
 DEFINE_LOG_CATEGORY(LogGameFeatures);
 
+const uint32 FInstallBundlePluginProtocolMetaData::FDefaultValues::CurrentVersionNum = 1;
+//Missing InstallBundles on purpose as the default is just an empty TArray and should always be encoded
+const bool FInstallBundlePluginProtocolMetaData::FDefaultValues::Default_bUninstallBeforeTerminate = true;
+const bool FInstallBundlePluginProtocolMetaData::FDefaultValues::Default_bUserPauseDownload = false;
+const EInstallBundleRequestFlags FInstallBundlePluginProtocolMetaData::FDefaultValues::Default_InstallBundleFlags = EInstallBundleRequestFlags::Defaults;
+
+namespace UE::GameFeatures
+{
+	static const FString SubsystemErrorNamespace(TEXT("GameFeaturePlugin.Subsystem."));
+
+	namespace PluginURLStructureInfo
+	{
+		const TCHAR* OptionAssignOperator = TEXT("=");
+		const TCHAR* OptionSeperator = TEXT("?");
+	}
+
+	namespace CommonErrorCodes
+	{
+		const FString PluginNotAllowed = TEXT("Plugin_Denied_By_GameSpecificPolicy");
+		const FString DependencyFailedRegister = TEXT("Failed_Dependency_Register");
+		const FString BadURL = TEXT("Bad_URL");
+		const FString UnreachableState = TEXT("State_Currently_Unreachable");
+		const FString NoURLUpdateNeeded = TEXT("URL_Not_Updated");
+
+		const FString CancelAddonCode = TEXT("_Cancel");
+
+		const FText GenericError = NSLOCTEXT("GameFeatures", "CommonErrors.Generic", "An error has occurred. Please try again later.");
+	}
+}
+
+#define GAME_FEATURE_PLUGIN_STATE_LEX_TO_STRING(inEnum, inText)  \
+    case(EGameFeaturePluginState::inEnum):                       \
+    {                                                            \
+        return TEXT("#inEnum");                                  \
+    }                                                            
+
+namespace GameFeaturePluginStatePrivate
+{
+	FString LexToString(EGameFeaturePluginState InEnum)
+	{
+		switch (InEnum)
+		{
+			GAME_FEATURE_PLUGIN_STATE_LIST(GAME_FEATURE_PLUGIN_STATE_LEX_TO_STRING)
+
+		default:
+			{
+				ensureAlwaysMsgf(false, TEXT("Logic error causing a missing LexToString value for EGameFeaturePluginState:%d"), static_cast<uint8>(InEnum));
+				return TEXT("ERROR_UNSUPPORTED_ENUM");
+			}
+		}
+	}
+}
+#undef GAME_FEATURE_PLUGIN_STATE_LEX_TO_STRING
+
+const FString LexToString(const EGameFeatureTargetState GameFeatureTargetState)
+{
+	switch (GameFeatureTargetState)
+	{
+		case EGameFeatureTargetState::Installed:
+			return TEXT("Installed");
+		case EGameFeatureTargetState::Registered:
+			return TEXT("Registered"); 
+		
+		case EGameFeatureTargetState::Loaded:
+			return TEXT("Loaded"); 
+		case EGameFeatureTargetState::Active:
+			return TEXT("Active"); 
+		default:
+			check(false);
+			return TEXT("Unknown");
+	}
+
+	static_assert((uint8)EGameFeatureTargetState::Count == 4, TEXT("Update LexToString to include new EGameFeatureTargetState"));
+}
+
+FGameFeaturePluginIdentifier::FGameFeaturePluginIdentifier(const FString& PluginURL)
+{
+	FromPluginURL(PluginURL);
+}
+
+void FGameFeaturePluginIdentifier::FromPluginURL(const FString& PluginURLIn)
+{
+	//Make sure we have no stale data
+	IdentifyingURLSubset.Reset();
+	PluginURL = PluginURLIn;
+
+	PluginProtocol = UGameFeaturesSubsystem::GetPluginURLProtocol(PluginURL);
+
+	if (ensureAlwaysMsgf( ((PluginProtocol != EGameFeaturePluginProtocol::Unknown) 
+						&& (PluginProtocol != EGameFeaturePluginProtocol::Count)),
+						TEXT("Invalid PluginProtocol in PluginURL %s"), *PluginURL))
+	{
+
+		int32 PluginProtocolEndIndex = FCString::Strlen(UE::GameFeatures::GameFeaturePluginProtocolPrefix(PluginProtocol));
+		int32 FirstOptionIndex = PluginURLIn.Find(UE::GameFeatures::PluginURLStructureInfo::OptionSeperator, ESearchCase::IgnoreCase, ESearchDir::FromStart, PluginProtocolEndIndex);
+		
+		//If we don't have any options, then the IdentifyingURLSubset is just our entire URL except the protocol string
+		if (FirstOptionIndex == INDEX_NONE)
+		{
+			IdentifyingURLSubset = PluginURL.RightChop(PluginProtocolEndIndex);
+		}
+		//The IdentifyingURLSubset will be the string between the end of the protocol string and before the first option
+		else
+		{
+			const int32 IdentifierCharCount = (FirstOptionIndex - PluginProtocolEndIndex);
+			IdentifyingURLSubset = PluginURL.Mid(PluginProtocolEndIndex, IdentifierCharCount);
+		}
+	}
+}
+
+bool FGameFeaturePluginIdentifier::operator==(const FGameFeaturePluginIdentifier& Other) const
+{
+	return ((PluginProtocol == Other.PluginProtocol) &&
+			(IdentifyingURLSubset.Equals(Other.IdentifyingURLSubset, ESearchCase::IgnoreCase)));
+}
+
+bool FGameFeaturePluginIdentifier::ExactMatchesURL(const FString& PluginURLIn) const
+{
+	return GetFullPluginURL().Equals(PluginURLIn, ESearchCase::IgnoreCase);
+}
+
 void FGameFeatureStateChangeContext::SetRequiredWorldContextHandle(FName Handle)
 {
 	WorldContextHandle = Handle;
@@ -170,8 +292,8 @@ void UGameFeaturesSubsystem::Initialize(FSubsystemCollectionBase& Collection)
 		ECVF_Cheat);
 
 	IConsoleManager::Get().RegisterConsoleCommand(
-		TEXT("UninstallGameFeaturePlugin"),
-		TEXT("Uninstalls a game feature plugin by URL"),
+		TEXT("ReleaseGameFeaturePlugin"),
+		TEXT("Releases a game feature plugin's InstallBundle data by URL"),
 		FConsoleCommandWithWorldArgsAndOutputDeviceDelegate::CreateLambda([](const TArray<FString>& Args, UWorld*, FOutputDevice& Ar)
 			{
 				if (Args.Num() > 0)
@@ -181,7 +303,7 @@ void UGameFeaturesSubsystem::Initialize(FSubsystemCollectionBase& Collection)
 					{
 						PluginURL = Args[0];
 					}
-					UGameFeaturesSubsystem::Get().UninstallGameFeaturePlugin(PluginURL);
+					UGameFeaturesSubsystem::Get().ReleaseGameFeaturePlugin(PluginURL);
 				}
 				else
 				{
@@ -370,15 +492,12 @@ FString UGameFeaturesSubsystem::GetPluginURL_FileProtocol(const FString& PluginD
 
 FString UGameFeaturesSubsystem::GetPluginURL_InstallBundleProtocol(const FString& PluginName, TArrayView<const FString> BundleNames)
 {
-	ensure(BundleNames.Num() > 0);
-
-	FString Path;
-	Path += TEXT("installbundle:");
-	Path += PluginName;
-	Path += TEXT("?");
-	Path += FString::Join(BundleNames, TEXT(","));
-
-	return Path;
+	FInstallBundlePluginProtocolMetaData ProtocolMetadata;
+	for (const FString& BundleName : BundleNames)
+	{
+		ProtocolMetadata.InstallBundles.Add(FName(BundleName));
+	}
+	return GetPluginURL_InstallBundleProtocol(PluginName, ProtocolMetadata);
 }
 
 FString UGameFeaturesSubsystem::GetPluginURL_InstallBundleProtocol(const FString& PluginName, const FString& BundleName)
@@ -388,15 +507,9 @@ FString UGameFeaturesSubsystem::GetPluginURL_InstallBundleProtocol(const FString
 
 FString UGameFeaturesSubsystem::GetPluginURL_InstallBundleProtocol(const FString& PluginName, const TArrayView<const FName> BundleNames)
 {
-	ensure(BundleNames.Num() > 0);
-
-	FString Path;
-	Path += TEXT("installbundle:");
-	Path += PluginName;
-	Path += TEXT("?");
-	Path += FString::JoinBy(BundleNames, TEXT(","), UE_PROJECTION_MEMBER(FName, ToString));
-
-	return Path;
+	FInstallBundlePluginProtocolMetaData ProtocolMetadata;
+	ProtocolMetadata.InstallBundles.Append(BundleNames.GetData(), BundleNames.Num());
+	return GetPluginURL_InstallBundleProtocol(PluginName, ProtocolMetadata);
 }
 
 FString UGameFeaturesSubsystem::GetPluginURL_InstallBundleProtocol(const FString& PluginName, FName BundleName)
@@ -404,6 +517,18 @@ FString UGameFeaturesSubsystem::GetPluginURL_InstallBundleProtocol(const FString
 	return GetPluginURL_InstallBundleProtocol(PluginName, MakeArrayView(&BundleName, 1));
 }
 
+FString UGameFeaturesSubsystem::GetPluginURL_InstallBundleProtocol(const FString& PluginName, const FInstallBundlePluginProtocolMetaData& ProtocolMetadata)
+{
+	ensure (ProtocolMetadata.InstallBundles.Num() > 0);
+	FString Path;
+	Path += UE::GameFeatures::GameFeaturePluginProtocolPrefix(EGameFeaturePluginProtocol::InstallBundle);
+	Path += PluginName;
+	Path += UE::GameFeatures::PluginURLStructureInfo::OptionSeperator;
+	Path += ProtocolMetadata.ToString();
+
+	return Path;
+}
+
 EGameFeaturePluginProtocol UGameFeaturesSubsystem::GetPluginURLProtocol(FStringView PluginURL)
 {
 	for (EGameFeaturePluginProtocol Protocol : TEnumRange<EGameFeaturePluginProtocol>())
@@ -506,6 +631,11 @@ void UGameFeaturesSubsystem::OnGameFeatureDeactivating(const UGameFeatureData* G
 	}
 }
 
+void UGameFeaturesSubsystem::OnGameFeaturePauseChange(const FString& PluginURL, const FString& PluginName, FGameFeaturePauseStateChangeContext& Context)
+{
+	CallbackObservers(EObserverCallback::PauseChanged, PluginURL, &PluginName, nullptr, &Context);
+}
+
 const UGameFeatureData* UGameFeaturesSubsystem::GetDataForStateMachine(UGameFeaturePluginStateMachine* GFSM) const
 {
 	return GFSM->GetGameFeatureDataForActivePlugin();
@@ -550,6 +680,15 @@ const UGameFeatureData* UGameFeaturesSubsystem::GetGameFeatureDataForRegisteredP
 	return nullptr;
 }
 
+bool UGameFeaturesSubsystem::IsGameFeaturePluginInstalled(const FString& PluginURL) const
+{
+	if (const UGameFeaturePluginStateMachine* StateMachine = FindGameFeaturePluginStateMachine(PluginURL))
+	{
+		return StateMachine->GetCurrentState() >= EGameFeaturePluginState::Installed;
+	}
+	return false;
+}
+
 bool UGameFeaturesSubsystem::IsGameFeaturePluginRegistered(const FString& PluginURL) const
 {
 	if (const UGameFeaturePluginStateMachine* StateMachine = FindGameFeaturePluginStateMachine(PluginURL))
@@ -573,7 +712,7 @@ void UGameFeaturesSubsystem::LoadGameFeaturePlugin(const FString& PluginURL, con
 	const bool bIsPluginAllowed = GameSpecificPolicies->IsPluginAllowed(PluginURL);
 	if (!bIsPluginAllowed)
 	{
-		CompleteDelegate.ExecuteIfBound(MakeError(TEXT("GameFeaturePlugin.StateMachine.Plugin_Denied_By_GameSpecificPolicy")));
+		CompleteDelegate.ExecuteIfBound(UE::GameFeatures::FResult(MakeError(UE::GameFeatures::SubsystemErrorNamespace + UE::GameFeatures::CommonErrorCodes::PluginNotAllowed)));
 		return;
 	}
 
@@ -621,7 +760,7 @@ void UGameFeaturesSubsystem::ChangeGameFeatureTargetState(const FString& PluginU
 		StateMachine = FindGameFeaturePluginStateMachine(PluginURL);
 		if (!StateMachine)
 		{
-			CompleteDelegate.ExecuteIfBound(MakeError(TEXT("GameFeaturePlugin.StateMachine.Plugin_Denied_By_GameSpecificPolicy")));
+			CompleteDelegate.ExecuteIfBound(UE::GameFeatures::FResult(MakeError(UE::GameFeatures::SubsystemErrorNamespace + UE::GameFeatures::CommonErrorCodes::PluginNotAllowed)));
 			return;
 		}
 	}
@@ -636,7 +775,7 @@ void UGameFeaturesSubsystem::ChangeGameFeatureTargetState(const FString& PluginU
 	{
 		if (TargetPluginState > StateMachine->GetCurrentState() || TargetPluginState > StateMachine->GetDestination())
 		{
-			CompleteDelegate.ExecuteIfBound(MakeError(TEXT("GameFeaturePlugin.StateMachine.Plugin_Denied_By_GameSpecificPolicy")));
+			CompleteDelegate.ExecuteIfBound(UE::GameFeatures::FResult(MakeError(UE::GameFeatures::SubsystemErrorNamespace + UE::GameFeatures::CommonErrorCodes::PluginNotAllowed)));
 			return;
 		}
 	}
@@ -653,9 +792,68 @@ void UGameFeaturesSubsystem::ChangeGameFeatureTargetState(const FString& PluginU
 		CallbackObservers(EObserverCallback::Activating, PluginURL, &StateMachine->GetPluginName(), StateMachine->GetGameFeatureDataForActivePlugin());
 	}
 	
+	if (ShouldUpdatePluginURLData(PluginURL))
+	{
+		UpdateGameFeaturePluginURL(PluginURL, FGameFeaturePluginUpdateURLComplete());
+	}
+
 	ChangeGameFeatureDestination(StateMachine, FGameFeaturePluginStateRange(TargetPluginState), CompleteDelegate);
 }
 
+void UGameFeaturesSubsystem::UpdateGameFeaturePluginURL(const FString& NewPluginURL)
+{
+	UpdateGameFeaturePluginURL(NewPluginURL, FGameFeaturePluginUpdateURLComplete());
+};
+
+void UGameFeaturesSubsystem::UpdateGameFeaturePluginURL(const FString& NewPluginURL, const FGameFeaturePluginUpdateURLComplete& CompleteDelegate)
+{
+	UGameFeaturePluginStateMachine* StateMachine = nullptr;
+	StateMachine = FindGameFeaturePluginStateMachine(NewPluginURL);
+	if (!StateMachine)
+	{
+		CompleteDelegate.ExecuteIfBound(MakeError(TEXT("GameFeaturePlugin.BadURL")));
+		return;
+	}
+
+	check(StateMachine);
+
+	const bool bUpdated = StateMachine->TryUpdatePluginURLData(NewPluginURL);
+	if (!bUpdated)
+	{
+		CompleteDelegate.ExecuteIfBound(MakeError(TEXT("GameFeaturePlugin.No_URL_Update_Needed")));
+	}
+	else
+	{
+		FString* PluginURL = GameFeaturePluginNameToPathMap.Find(StateMachine->GetPluginName());
+		if (!ensureAlwaysMsgf(PluginURL, TEXT("Attempt to UpdateGameFeaturePluginURL before GameFeaturePlugin has been added to the GameFeaturePluginNameToPathMap! URL:%s"), *NewPluginURL))
+		{
+			CompleteDelegate.ExecuteIfBound(MakeError(TEXT("GameFeaturePlugin.UpdateTooEarly")));
+			return;
+		}
+
+		*PluginURL = NewPluginURL;
+		CompleteDelegate.ExecuteIfBound(MakeValue());
+	}
+}
+
+bool UGameFeaturesSubsystem::ShouldUpdatePluginURLData(const FString& NewPluginURL)
+{
+	UGameFeaturePluginStateMachine* StateMachine = nullptr;
+	StateMachine = FindGameFeaturePluginStateMachine(NewPluginURL);
+	if (!StateMachine || !StateMachine->IsStatusKnown())
+	{
+		return false;
+	}
+
+	check(StateMachine);
+	if (StateMachine->GetPluginURL().Equals(NewPluginURL, ESearchCase::IgnoreCase))
+	{
+		return false;
+	}
+
+	return true;
+}
+
 bool UGameFeaturesSubsystem::GetGameFeaturePluginInstallPercent(const FString& PluginURL, float& Install_Percent) const
 {
 	if (const UGameFeaturePluginStateMachine* StateMachine = FindGameFeaturePluginStateMachine(PluginURL))
@@ -706,7 +904,7 @@ void UGameFeaturesSubsystem::DeactivateGameFeaturePlugin(const FString& PluginUR
 	}
 	else
 	{
-		CompleteDelegate.ExecuteIfBound(UE::GameFeatures::FResult(MakeError(TEXT("GameFeaturePlugin.BadURL"))));
+		CompleteDelegate.ExecuteIfBound(UE::GameFeatures::FResult(MakeError(UE::GameFeatures::SubsystemErrorNamespace + UE::GameFeatures::CommonErrorCodes::BadURL)));
 	}
 }
 
@@ -724,16 +922,16 @@ void UGameFeaturesSubsystem::UnloadGameFeaturePlugin(const FString& PluginURL, c
 	}
 	else
 	{
-		CompleteDelegate.ExecuteIfBound(UE::GameFeatures::FResult(MakeError(TEXT("GameFeaturePlugin.BadURL"))));
+		CompleteDelegate.ExecuteIfBound(UE::GameFeatures::FResult(MakeError(UE::GameFeatures::SubsystemErrorNamespace + UE::GameFeatures::CommonErrorCodes::BadURL)));
 	}
 }
 
-void UGameFeaturesSubsystem::UninstallGameFeaturePlugin(const FString& PluginURL)
+void UGameFeaturesSubsystem::ReleaseGameFeaturePlugin(const FString& PluginURL)
 {
-	UninstallGameFeaturePlugin(PluginURL, FGameFeaturePluginUninstallComplete());
+	ReleaseGameFeaturePlugin(PluginURL, FGameFeaturePluginReleaseComplete());
 }
 
-void UGameFeaturesSubsystem::UninstallGameFeaturePlugin(const FString& PluginURL, const FGameFeaturePluginUninstallComplete& CompleteDelegate)
+void UGameFeaturesSubsystem::ReleaseGameFeaturePlugin(const FString& PluginURL, const FGameFeaturePluginReleaseComplete& CompleteDelegate)
 {
 	if (UGameFeaturePluginStateMachine* StateMachine = FindGameFeaturePluginStateMachine(PluginURL))
 	{
@@ -741,24 +939,78 @@ void UGameFeaturesSubsystem::UninstallGameFeaturePlugin(const FString& PluginURL
 	}
 	else
 	{
-		CompleteDelegate.ExecuteIfBound(UE::GameFeatures::FResult(MakeError(TEXT("GameFeaturePlugin.BadURL"))));
+		CompleteDelegate.ExecuteIfBound(UE::GameFeatures::FResult(MakeError(UE::GameFeatures::SubsystemErrorNamespace + UE::GameFeatures::CommonErrorCodes::BadURL)));
 	}
 }
 
+void UGameFeaturesSubsystem::UninstallGameFeaturePlugin(const FString& PluginURL, const FGameFeaturePluginUninstallComplete& CompleteDelegate)
+{
+	UGameFeaturePluginStateMachine* StateMachine = FindOrCreateGameFeaturePluginStateMachine(PluginURL);
+	check(StateMachine);
+
+	//Define a lambda wrapper to return the FResult of the Terminate for our FGameFeaturePluginUninstallComplete Delegate
+	const FGameFeaturePluginTerminateComplete ReportTerminateResultLambda = FGameFeaturePluginTerminateComplete::CreateLambda([=](const UE::GameFeatures::FResult& Result)
+		{
+			CompleteDelegate.Execute(Result);
+		});
+
+	//We may need to update our PluginURL to force certain metadata changes to facilitate this uninstall
+	//This tracks what URL we actually will pass in to the terminate as we can rely on the fact that Terminate
+	//will update any important Metadata from this new URL before beginning it's terminate
+	FString PluginURLForTerminate = PluginURL;
+
+	//InstallBundle Protocol GameFeatures may need to change their metadata to force this uninstall
+	if (EGameFeaturePluginProtocol::InstallBundle == GetPluginURLProtocol(PluginURL))
+	{
+		// Parse a duplicate version of our current Metadata from the URL
+		FInstallBundlePluginProtocolMetaData ProtocolMetadata;
+		if (!FInstallBundlePluginProtocolMetaData::FromString(PluginURL, ProtocolMetadata))
+		{
+			CompleteDelegate.ExecuteIfBound(UE::GameFeatures::FResult(MakeError(UE::GameFeatures::SubsystemErrorNamespace + UE::GameFeatures::CommonErrorCodes::BadURL)));
+		}
+
+		// Need to force on bUninstallBeforeTerminate if it wasn't already set to on in our Metadata
+		if (!ProtocolMetadata.bUninstallBeforeTerminate)
+		{
+			ProtocolMetadata.bUninstallBeforeTerminate = true;
+			PluginURLForTerminate = GetPluginURL_InstallBundleProtocol(StateMachine->GetPluginName(), ProtocolMetadata);
+		}
+	}
+
+	TerminateGameFeaturePlugin(PluginURLForTerminate, ReportTerminateResultLambda);
+}
+
 void UGameFeaturesSubsystem::TerminateGameFeaturePlugin(const FString& PluginURL)
 {
-	TerminateGameFeaturePlugin(PluginURL, FGameFeaturePluginUninstallComplete());
+	TerminateGameFeaturePlugin(PluginURL, FGameFeaturePluginTerminateComplete());
 }
 
-void UGameFeaturesSubsystem::TerminateGameFeaturePlugin(const FString& PluginURL, const FGameFeaturePluginUninstallComplete& CompleteDelegate)
+void UGameFeaturesSubsystem::TerminateGameFeaturePlugin(const FString& PluginURL, const FGameFeaturePluginTerminateComplete& CompleteDelegate)
 {
 	if (UGameFeaturePluginStateMachine* StateMachine = FindGameFeaturePluginStateMachine(PluginURL))
 	{
-		ChangeGameFeatureDestination(StateMachine, FGameFeaturePluginStateRange(EGameFeaturePluginState::Terminal), CompleteDelegate);
+		//Define a lambda that will kick off the actual Terminate
+		const FGameFeaturePluginUpdateURLComplete StartTerminateLambda = FGameFeaturePluginUpdateURLComplete::CreateLambda([=](const UE::GameFeatures::FResult& Result)
+			{
+				UGameFeaturePluginStateMachine* StateMachine = FindGameFeaturePluginStateMachine(PluginURL);
+				check(StateMachine);
+				ChangeGameFeatureDestination(StateMachine, FGameFeaturePluginStateRange(EGameFeaturePluginState::Terminal), CompleteDelegate);
+			});
+
+		//If we need to update our URLData, try to do that first before starting the Terminate. This allows us to update
+		//URL Metadata flags that might be important on the way to Terminal if they are changed. EX: FInstallBundlePluginProtocolMetaData::bUninstallBeforeTerminate
+		if (ShouldUpdatePluginURLData(PluginURL))
+		{
+			UpdateGameFeaturePluginURL(PluginURL, StartTerminateLambda);
+		}
+		else
+		{
+			StartTerminateLambda.Execute(MakeValue());
+		}
 	}
 	else
 	{
-		CompleteDelegate.ExecuteIfBound(UE::GameFeatures::FResult(MakeError(TEXT("GameFeaturePlugin.BadURL"))));
+		CompleteDelegate.ExecuteIfBound(UE::GameFeatures::FResult(MakeError(UE::GameFeatures::SubsystemErrorNamespace + UE::GameFeatures::CommonErrorCodes::BadURL)));
 	}
 }
 
@@ -783,7 +1035,7 @@ void UGameFeaturesSubsystem::CancelGameFeatureStateChange(const FString& PluginU
 	}
 	else
 	{
-		CompleteDelegate.ExecuteIfBound(UE::GameFeatures::FResult(MakeError(TEXT("GameFeaturePlugin.BadURL"))));
+		CompleteDelegate.ExecuteIfBound(UE::GameFeatures::FResult(MakeError(UE::GameFeatures::SubsystemErrorNamespace + UE::GameFeatures::CommonErrorCodes::BadURL)));
 	}
 }
 
@@ -1112,7 +1364,7 @@ void UGameFeaturesSubsystem::ChangeGameFeatureDestination(UGameFeaturePluginStat
 			{
 				UE_LOG(LogGameFeatures, Warning, TEXT("ChangeGameFeatureDestination: Failed to set Game Feature %s Destination State to [%s, %s]"), *Machine->GetGameFeatureName(), *UE::GameFeatures::ToString(StateRange.MinState), *UE::GameFeatures::ToString(StateRange.MaxState));
 
-				CompleteDelegate.ExecuteIfBound(MakeError(TEXT("GameFeaturePlugin.StateMachine.State_Currently_Unreachable")));
+				CompleteDelegate.ExecuteIfBound(UE::GameFeatures::FResult(MakeError(UE::GameFeatures::SubsystemErrorNamespace + UE::GameFeatures::CommonErrorCodes::UnreachableState)));
 			}
 			else
 			{
@@ -1125,7 +1377,7 @@ void UGameFeaturesSubsystem::ChangeGameFeatureDestination(UGameFeaturePluginStat
 		{
 			UE_LOG(LogGameFeatures, Warning, TEXT("ChangeGameFeatureDestination: Failed to cancel Game Feature %s"), *Machine->GetGameFeatureName());
 
-			CompleteDelegate.ExecuteIfBound(MakeError(TEXT("GameFeaturePlugin.StateMachine.State_Currently_Unreachable_Cancel_Failed")));
+			CompleteDelegate.ExecuteIfBound(UE::GameFeatures::FResult(MakeError(UE::GameFeatures::SubsystemErrorNamespace + UE::GameFeatures::CommonErrorCodes::UnreachableState + UE::GameFeatures::CommonErrorCodes::CancelAddonCode)));
 		}
 	}
 }
@@ -1218,28 +1470,33 @@ void UGameFeaturesSubsystem::ListGameFeaturePlugins(const TArray<FString>& Args,
 void UGameFeaturesSubsystem::CallbackObservers(EObserverCallback CallbackType, const FString& PluginURL, 
 	const FString* PluginName /*= nullptr*/, 
 	const UGameFeatureData* GameFeatureData /*= nullptr*/, 
-	FGameFeatureDeactivatingContext* DeactivatingContext /*= nullptr*/)
+	FGameFeatureStateChangeContext* StateChangeContext /*= nullptr*/)
 {
+	static_assert(std::underlying_type<EObserverCallback>::type(EObserverCallback::Count) == 8, "Update UGameFeaturesSubsystem::CallbackObservers to handle added EObserverCallback");
+
 	// Protect against modifying the observer list during iteration
 	TArray<UObject*> LocalObservers(Observers);
 
 	switch (CallbackType)
 	{
 	case EObserverCallback::CheckingStatus:
+	{
 		for (UObject* Observer : LocalObservers)
 		{
 			CastChecked<IGameFeatureStateChangeObserver>(Observer)->OnGameFeatureCheckingStatus(PluginURL);
 		}
 		break;
-
+	}
 	case EObserverCallback::Terminating:
+	{
 		for (UObject* Observer : LocalObservers)
 		{
 			CastChecked<IGameFeatureStateChangeObserver>(Observer)->OnGameFeatureTerminating(PluginURL);
 		}
 		break;
-
+	}
 	case EObserverCallback::Registering:
+	{
 		check(PluginName);
 		check(GameFeatureData);
 		for (UObject* Observer : LocalObservers)
@@ -1247,8 +1504,9 @@ void UGameFeaturesSubsystem::CallbackObservers(EObserverCallback CallbackType, c
 			CastChecked<IGameFeatureStateChangeObserver>(Observer)->OnGameFeatureRegistering(GameFeatureData, *PluginName, PluginURL);
 		}
 		break;
-
+	}
 	case EObserverCallback::Unregistering:
+	{
 		check(PluginName);
 		check(GameFeatureData);
 		for (UObject* Observer : LocalObservers)
@@ -1256,32 +1514,53 @@ void UGameFeaturesSubsystem::CallbackObservers(EObserverCallback CallbackType, c
 			CastChecked<IGameFeatureStateChangeObserver>(Observer)->OnGameFeatureUnregistering(GameFeatureData, *PluginName, PluginURL);
 		}
 		break;
-
+	}
 	case EObserverCallback::Loading:
+	{
 		check(GameFeatureData);
 		for (UObject* Observer : LocalObservers)
 		{
 			CastChecked<IGameFeatureStateChangeObserver>(Observer)->OnGameFeatureLoading(GameFeatureData, PluginURL);
 		}
 		break;
-
+	}
 	case EObserverCallback::Activating:
+	{
 		check(GameFeatureData);
 		for (UObject* Observer : LocalObservers)
 		{
 			CastChecked<IGameFeatureStateChangeObserver>(Observer)->OnGameFeatureActivating(GameFeatureData, PluginURL);
 		}
 		break;
-
+	}
 	case EObserverCallback::Deactivating:
+	{
 		check(GameFeatureData);
-		check(DeactivatingContext);
-		for (UObject* Observer : LocalObservers)
+		check(StateChangeContext);
+		FGameFeatureDeactivatingContext* DeactivatingContext = static_cast<FGameFeatureDeactivatingContext*>(StateChangeContext);
+		if (ensureAlwaysMsgf(DeactivatingContext, TEXT("Invalid StateChangeContext supplied! Could not cast to FGameFeaturePauseStateChangeContext*!")))
+		{
+			for (UObject* Observer : LocalObservers)
+			{
+				CastChecked<IGameFeatureStateChangeObserver>(Observer)->OnGameFeatureDeactivating(GameFeatureData, *DeactivatingContext, PluginURL);
+			}
+		}
+		break;
+	}
+	case EObserverCallback::PauseChanged:
+	{
+		check(PluginName);
+		check(StateChangeContext);
+		FGameFeaturePauseStateChangeContext* PauseChangeContext = static_cast<FGameFeaturePauseStateChangeContext*>(StateChangeContext);
+		if (ensureAlwaysMsgf(PauseChangeContext, TEXT("Invalid StateChangeContext supplied! Could not cast to FGameFeaturePauseStateChangeContext*!")))
 		{
-			CastChecked<IGameFeatureStateChangeObserver>(Observer)->OnGameFeatureDeactivating(GameFeatureData, *DeactivatingContext, PluginURL);
+			for (UObject* Observer : LocalObservers)
+			{
+				CastChecked<IGameFeatureStateChangeObserver>(Observer)->OnGameFeaturePauseChange(PluginURL, *PluginName, *PauseChangeContext);
+			}
 		}
 		break;
-
+	}
 	default:
 		UE_LOG(LogGameFeatures, Fatal, TEXT("Unkown EObserverCallback!"));
 	}
@@ -1291,7 +1570,7 @@ TSet<FString> UGameFeaturesSubsystem::GetActivePluginNames() const
 {
 	TSet<FString> ActivePluginNames;
 
-	for (const UE_TRANSITIONAL_OBJECT_PTR_TEMPLATE2_ARG2(TPair, FString, UGameFeaturePluginStateMachine)& Pair : GameFeaturePluginStateMachines)
+	for (const UE_TRANSITIONAL_OBJECT_PTR_TEMPLATE2_ARG2(TPair, FGameFeaturePluginIdentifier, UGameFeaturePluginStateMachine)& Pair : GameFeaturePluginStateMachines)
 	{
 		UGameFeaturePluginStateMachine* StateMachine = Pair.Value;
 		if (StateMachine->GetCurrentState() == EGameFeaturePluginState::Active &&
diff --git a/Engine/Plugins/Experimental/GameFeatures/Source/GameFeatures/Public/GameFeaturePluginOperationResult.h b/Engine/Plugins/Experimental/GameFeatures/Source/GameFeatures/Public/GameFeaturePluginOperationResult.h
--- a/Engine/Plugins/Experimental/GameFeatures/Source/GameFeatures/Public/GameFeaturePluginOperationResult.h
+++ b/Engine/Plugins/Experimental/GameFeatures/Source/GameFeatures/Public/GameFeaturePluginOperationResult.h
@@ -7,9 +7,54 @@
 
 namespace UE::GameFeatures
 {
-	struct FSuccessType {};
+	//Type used to determine if our FResult is actually an error or success (and hold the error code)
+	using FErrorCodeType = TValueOrError<void, FString>;
 
-	using FResult = TValueOrError<FSuccessType, FString>;
+	/** Struct wrapper on a combined TValueOrError ErrorCode and FText OptionalErrorText that can provide an FString error code and
+		if set more detailed FText OptionalErrorText information for the results of a GameFeaturesPlugin operation. */
+	struct GAMEFEATURES_API FResult
+	{
+		/** Error Code representing the error that occurred */
+		FErrorCodeType ErrorCode;
+
+		/** Optional Localized error description to bubble up to the user if one was generated */
+		FText OptionalErrorText;
+
+		/** Quick functions that just pass through to the ErrorCode */
+		bool HasValue() const { return ErrorCode.HasValue(); }
+		bool HasError() const { return ErrorCode.HasError(); }
+		FString GetError() const { return ErrorCode.GetError(); }
+		FString StealError() { return ErrorCode.StealError(); }
+
+		FResult(FErrorCodeType ErrorCodeIn, FText ErrorTextIn);
+
+		/** Explicit constructor for times we want to create an FResult directly from an FErrorCodeType and not through MakeValue or MakeError.
+			Explicit to avoid any confusion with the following templated constructors. */
+		explicit FResult(FErrorCodeType ErrorCodeIn);
+
+		/** Template Conversion Constructor to allow us to initialize from TValueOrError MakeValue
+			This is needed because of how TValueOrError implements MakeValue through the same templated approach with TValueOrError_ValueProxy. */
+		template <typename... ArgTypes>
+		FResult(TValueOrError_ValueProxy<ArgTypes...>&& ValueProxyIn)
+			: ErrorCode(MoveTemp(ValueProxyIn))
+			, OptionalErrorText()
+		{
+		}
+
+		/** Template Conversion Constructor to allow us to initialize from TValueOrError MakeError 
+			This is needed because of how TValueOrError implements MakeError through the same templated approach with TValueOrError_ErrorProxy. */
+		template <typename... ArgTypes>
+		FResult(TValueOrError_ErrorProxy<ArgTypes...>&& ErrorProxyIn)
+			: ErrorCode(MoveTemp(ErrorProxyIn))
+			, OptionalErrorText()
+		{
+		}
+
+	private:
+		//No default constructor as we want to force you to always specify at the minimum 
+		//if the FResult is an error or not through a supplied FErrorCodeType
+		FResult() = delete;
+	};
 
 	GAMEFEATURES_API FString ToString(const FResult& Result);
 }
diff --git a/Engine/Plugins/Experimental/GameFeatures/Source/GameFeatures/Public/GameFeatureStateChangeObserver.h b/Engine/Plugins/Experimental/GameFeatures/Source/GameFeatures/Public/GameFeatureStateChangeObserver.h
--- a/Engine/Plugins/Experimental/GameFeatures/Source/GameFeatures/Public/GameFeatureStateChangeObserver.h
+++ b/Engine/Plugins/Experimental/GameFeatures/Source/GameFeatures/Public/GameFeatureStateChangeObserver.h
@@ -8,6 +8,7 @@
 
 class UGameFeatureData;
 struct FGameFeatureDeactivatingContext;
+struct FGameFeaturePauseStateChangeContext;
 
 UINTERFACE(MinimalAPI)
 class UGameFeatureStateChangeObserver : public UInterface
@@ -44,4 +45,9 @@ public:
 	virtual void OnGameFeatureActivating(const UGameFeatureData* GameFeatureData, const FString& PluginURL) {}
 
 	virtual void OnGameFeatureDeactivating(const UGameFeatureData* GameFeatureData, FGameFeatureDeactivatingContext& Context, const FString& PluginURL) {}
+
+	/** Called whenever a GameFeature State either pauses or resumes work without transitioning out of that state.
+		EX: Downloading paused due to a users cellular data settings or the user taking a pause action. We
+		may not yet transition to a download error, but want a way to observe this behavior. */
+	virtual void OnGameFeaturePauseChange(const FString& PluginURL, const FString& PluginName, FGameFeaturePauseStateChangeContext& Context) {}
 };
diff --git a/Engine/Plugins/Experimental/GameFeatures/Source/GameFeatures/Public/GameFeatureTypes.h b/Engine/Plugins/Experimental/GameFeatures/Source/GameFeatures/Public/GameFeatureTypes.h
--- a/Engine/Plugins/Experimental/GameFeatures/Source/GameFeatures/Public/GameFeatureTypes.h
+++ b/Engine/Plugins/Experimental/GameFeatures/Source/GameFeatures/Public/GameFeatureTypes.h
@@ -8,12 +8,14 @@
 	XSTATE(Uninitialized,				NSLOCTEXT("GameFeatures", "UninitializedStateDisplayName", "Uninitialized"))								/* Unset. Not yet been set up. */ \
 	XSTATE(Terminal,					NSLOCTEXT("GameFeatures", "TerminalStateDisplayName", "Terminal"))											/* Final State before removal of the state machine. */ \
 	XSTATE(UnknownStatus,				NSLOCTEXT("GameFeatures", "UnknownStatusStateDisplayName", "UnknownStatus"))								/* Initialized, but the only thing known is the URL to query status. */ \
+	XSTATE(Uninstalling,				NSLOCTEXT("GameFeatures", "UninstallingStateDisplayName", "Uninstalling"))									/* Transition state between StatusKnown -> Terminal for any plugin that can have data that needs to have local datat uninstalled. */ \
+	XSTATE(ErrorUninstalling,			NSLOCTEXT("GameFeatures", "ErrorUninstallingStateDisplayName", "ErrorUninstalling"))						/* Error state for Uninstalling -> Terminal transition. */  \
 	XSTATE(CheckingStatus,				NSLOCTEXT("GameFeatures", "CheckingStatusStateDisplayName", "CheckingStatus"))								/* Transition state UnknownStatus -> StatusKnown. The status is in the process of being queried. */ \
 	XSTATE(ErrorCheckingStatus,			NSLOCTEXT("GameFeatures", "ErrorCheckingStatusStateDisplayName", "ErrorCheckingStatus"))					/* Error state for UnknownStatus -> StatusKnown transition. */ \
 	XSTATE(ErrorUnavailable,			NSLOCTEXT("GameFeatures", "ErrorUnavailableStateDisplayName", "ErrorUnavailable"))							/* Error state for UnknownStatus -> StatusKnown transition. */ \
 	XSTATE(StatusKnown,					NSLOCTEXT("GameFeatures", "StatusKnownStateDisplayName", "StatusKnown"))									/* The plugin's information is known, but no action has taken place yet. */ \
-	XSTATE(ErrorInstalling,				NSLOCTEXT("GameFeatures", "ErrorInstallingStateDisplayName", "ErrorInstalling"))							/* Error state for Installed -> StatusKnown and StatusKnown -> Installed transitions. */ \
-	XSTATE(Uninstalling,				NSLOCTEXT("GameFeatures", "UninstallingStateDisplayName", "Uninstalling"))									/* Transition state Installed -> StatusKnown. In the process of removing from local storage. */ \
+	XSTATE(Releasing,					NSLOCTEXT("GameFeatures", "ReleasingStateDisplayName", "Releasing"))										/* Transition State for Installed -> StatusKnown. Releases local data from any relevant caches. */ \
+	XSTATE(ErrorManagingData,			NSLOCTEXT("GameFeatures", "ErrorManagingDataStateDisplayName", "ErrorManagingData"))						/* Error state for Installed -> StatusKnown and StatusKnown -> Installed transitions. */ \
 	XSTATE(Downloading,					NSLOCTEXT("GameFeatures", "DownloadingStateDisplayName", "Downloading"))									/* Transition state StatusKnown -> Installed. In the process of adding to local storage. */ \
 	XSTATE(Installed,					NSLOCTEXT("GameFeatures", "InstalledStateDisplayName", "Installed"))										/* The plugin is in local storage (i.e. it is on the hard drive) */ \
 	XSTATE(ErrorMounting,				NSLOCTEXT("GameFeatures", "ErrorMountingStateDisplayName", "ErrorMounting"))								/* Error state for Installed -> Registered and Registered -> Installed transitions. */ \
@@ -40,6 +42,9 @@ namespace GameFeaturePluginStatePrivate
 		GAME_FEATURE_PLUGIN_STATE_LIST(GAME_FEATURE_PLUGIN_STATE_ENUM)
 		MAX
 	};
+
+	//defined in GameFeaturesSubsystem.cpp
+	FString LexToString(EGameFeaturePluginState InEnum);
 }
 typedef GameFeaturePluginStatePrivate::EGameFeaturePluginState EGameFeaturePluginState;
 #undef GAME_FEATURE_PLUGIN_STATE_ENUM
diff --git a/Engine/Plugins/Experimental/GameFeatures/Source/GameFeatures/Public/GameFeaturesSubsystem.h b/Engine/Plugins/Experimental/GameFeatures/Source/GameFeatures/Public/GameFeaturesSubsystem.h
--- a/Engine/Plugins/Experimental/GameFeatures/Source/GameFeatures/Public/GameFeaturesSubsystem.h
+++ b/Engine/Plugins/Experimental/GameFeatures/Source/GameFeatures/Public/GameFeaturesSubsystem.h
@@ -20,6 +20,33 @@ class IPlugin;
 class FJsonObject;
 struct FWorldContext;
 struct FGameFeaturePluginStateRange;
+enum class EInstallBundleRequestFlags : uint32;
+
+/** Holds static global information about how our PluginURLs are structured */
+namespace UE::GameFeatures
+{
+	namespace PluginURLStructureInfo
+	{
+		extern const TCHAR* OptionAssignOperator;
+		extern const TCHAR* OptionSeperator;
+	};
+
+	namespace CommonErrorCodes
+	{
+		extern const FString PluginNotAllowed;
+		extern const FString DependencyFailedRegister;
+		extern const FString BadURL;
+		extern const FString UnreachableState;
+		extern const FString NoURLUpdateNeeded;
+		
+		extern const FString CancelAddonCode;
+	};
+
+	namespace CommonErrorText
+	{
+		extern const FText GenericError;
+	};
+};
 
 /** 
  * Struct that determines if game feature action state changes should be applied for cases where there are multiple worlds or contexts.
@@ -95,6 +122,32 @@ private:
 	friend struct FGameFeaturePluginState_Deactivating;
 };
 
+/** Context that provides extra information for a game feature changing its pause state */
+struct FGameFeaturePauseStateChangeContext : public FGameFeatureStateChangeContext
+{
+public:
+	FGameFeaturePauseStateChangeContext(FString PauseStateNameIn, FString PauseReasonIn, bool bIsPausedIn)
+		: PauseStateName(MoveTemp(PauseStateNameIn))
+		, PauseReason(MoveTemp(PauseReasonIn))
+		, bIsPaused(bIsPausedIn)
+	{
+	}
+
+	/** Returns true if the State has paused or false if it is resuming */
+	bool IsPaused() const { return bIsPaused; }
+
+	/** Returns an FString description of why the state has paused work. */
+	const FString& GetPauseReason() const { return PauseReason; }
+
+	/** Returns an FString description of what state has issued the pause change */
+	const FString& GetPausingStateName() const { return PauseStateName; }
+
+private:
+	FString PauseStateName;
+	FString PauseReason;
+	bool bIsPaused = false;
+};
+
 GAMEFEATURES_API DECLARE_LOG_CATEGORY_EXTERN(LogGameFeatures, Log, All);
 /** Notification that a game feature plugin install/register/load/unload has finished */
 DECLARE_DELEGATE_OneParam(FGameFeaturePluginChangeStateComplete, const UE::GameFeatures::FResult& /*Result*/);
@@ -102,7 +155,15 @@ DECLARE_DELEGATE_OneParam(FGameFeaturePluginChangeStateComplete, const UE::GameF
 using FGameFeaturePluginLoadComplete = FGameFeaturePluginChangeStateComplete;
 using FGameFeaturePluginDeactivateComplete = FGameFeaturePluginChangeStateComplete;
 using FGameFeaturePluginUnloadComplete = FGameFeaturePluginChangeStateComplete;
+using FGameFeaturePluginReleaseComplete = FGameFeaturePluginChangeStateComplete;
 using FGameFeaturePluginUninstallComplete = FGameFeaturePluginChangeStateComplete;
+using FGameFeaturePluginTerminateComplete = FGameFeaturePluginChangeStateComplete;
+using FGameFeaturePluginUpdateURLComplete = FGameFeaturePluginChangeStateComplete;
+
+/** Notification delegate to be called by certain states that want to notify when they pause/resume work without leaving their current state.
+EX: Downloading can be paused due to cellular or internet connection outages without failing the download,
+but bubbling that pause information up to the user can be beneficial for messaging */
+DECLARE_DELEGATE_TwoParams(FGameFeaturePluginOnStatePausedChange, bool /*bIsPaused*/, const FString& /*PauseReason*/);
 
 /** Notification that a game feature plugin load has finished successfully and feeds back the GameFeatureData*/
 DECLARE_MULTICAST_DELEGATE_TwoParams(FGameFeaturePluginLoadCompleteDataReady, const FString& /*Name*/, const UGameFeatureData* /*Data*/);
@@ -128,6 +189,7 @@ enum class EGameFeatureTargetState : uint8
 	Active,
 	Count	UMETA(Hidden)
 };
+const FString GAMEFEATURES_API LexToString(const EGameFeatureTargetState GameFeatureTargetState);
 
 struct FGameFeaturePluginDetails
 {
@@ -147,7 +209,89 @@ struct FBuiltInGameFeaturePluginBehaviorOptions
 	EBuiltInAutoState AutoStateOverride = EBuiltInAutoState::Invalid;
 };
 
+/** Struct used to transform a GameFeaturePlugin URL into something that can uniquely identify the GameFeaturePlugin
+    without including any transient data being passed in through the URL */
+USTRUCT()
+struct FGameFeaturePluginIdentifier
+{
+	GENERATED_BODY()
+
+	/** The protocol used in the URL for this GameFeaturePlugin URL */
+	EGameFeaturePluginProtocol PluginProtocol;
+
+	FGameFeaturePluginIdentifier() = default;
+	FGameFeaturePluginIdentifier(const FString& PluginURL);
+
+	/** Used to determine if 2 FGameFeaturePluginIdentifiers are referencing the same GameFeaturePlugin.
+		Only matching on Identifying information instead of all the optional bundle information */
+	bool operator==(const FGameFeaturePluginIdentifier& Other) const;
 
+	/** Function that fills out IdentifyingURLSubset from the given PluginURL */
+	void FromPluginURL(const FString& PluginURL);
+
+	/** Returns true if this FGameFeaturePluginIdentifier exactly matches the given PluginURL.
+		To match exactly all information in the PluginURL has to match and not just the IdentifyingURLSubset */
+	bool ExactMatchesURL(const FString& PluginURL) const;
+
+	/** Get the Full PluginURL used to originally construct this identifier */
+	const FString& GetFullPluginURL() const { return PluginURL; }
+
+	friend FORCEINLINE uint32 GetTypeHash(const FGameFeaturePluginIdentifier& PluginIdentifier)
+	{
+		return GetTypeHash(PluginIdentifier.IdentifyingURLSubset);
+	}
+
+private:
+	/** The part of the URL that can be used to uniquely identify this plugin without any transient data */
+	FString IdentifyingURLSubset;
+
+	/** Full PluginURL used to originally construct this identifier */
+	FString PluginURL;
+
+	//Friend class so that it can access parsed URL data from under the hood
+	friend struct FGameFeaturePluginStateMachineProperties;
+};
+
+USTRUCT()
+struct GAMEFEATURES_API FInstallBundlePluginProtocolMetaData
+{
+	GENERATED_BODY()
+
+	TArray<FName> InstallBundles;
+
+	/** Set to whatever the FDefaultValues::CurrentVersionNum was when this URL was generated.
+		Allows us to ensure if we try and load URLs generated from a previous version **/
+	uint8 VersionNum;
+
+	/** If we want to attempt to uninstall InstallBundle information installed by this plugin before terminating */
+	bool bUninstallBeforeTerminate;
+
+	/** EInstallBundleRequestFlags utilized during the download/install by InstallBundleManager */
+	EInstallBundleRequestFlags InstallBundleFlags;
+
+	/** If we want to set the Downloading state to pause because of user interaction */
+	bool bUserPauseDownload;
+
+	/** Functions to convert to/from the URL FString representation of this metadata **/
+	FString ToString() const;
+	static bool FromString(const FString& URLString, FInstallBundlePluginProtocolMetaData& OutMetadata);
+
+	/** Resets all our Metadata values to the default values */
+	void ResetToDefaults();
+
+	FInstallBundlePluginProtocolMetaData();
+
+private:
+	/** Holds default values for the above settings as they are only encoded into a string if they differ from these values */
+	struct FDefaultValues
+	{
+		static const uint32 CurrentVersionNum;
+		//Missing InstallBundles on purpose as the default is just an empty TArray and should always be encoded
+		static const bool Default_bUninstallBeforeTerminate;
+		static const bool Default_bUserPauseDownload;
+		static const EInstallBundleRequestFlags Default_InstallBundleFlags;
+	};
+};
 
 /** The manager subsystem for game features */
 UCLASS()
@@ -218,6 +362,7 @@ public:
 	static FString GetPluginURL_InstallBundleProtocol(const FString& PluginName, const FString& BundleName);
 	static FString GetPluginURL_InstallBundleProtocol(const FString& PluginName, TArrayView<const FName> BundleNames);
 	static FString GetPluginURL_InstallBundleProtocol(const FString& PluginName, FName BundleName);
+	static FString GetPluginURL_InstallBundleProtocol(const FString& PluginName, const FInstallBundlePluginProtocolMetaData& ProtocolMetadata);
 
 	static EGameFeaturePluginProtocol GetPluginURLProtocol(FStringView PluginURL);
 
@@ -231,6 +376,9 @@ public:
 	/** Returns the game feature data for a registered plugin specified by PluginURL */
 	const UGameFeatureData* GetGameFeatureDataForRegisteredPluginByURL(const FString& PluginURL);
 
+	/** Determines if a plugin is in the Installed state (or beyond) */
+	bool IsGameFeaturePluginInstalled(const FString& PluginURL) const;
+
 	/** Determines if a plugin is in the Registered state (or beyond) */
 	bool IsGameFeaturePluginRegistered(const FString& PluginURL) const;
 
@@ -246,6 +394,10 @@ public:
 	/** Changes the target state of a game feature plugin */
 	void ChangeGameFeatureTargetState(const FString& PluginURL, EGameFeatureTargetState TargetState, const FGameFeaturePluginChangeStateComplete& CompleteDelegate);
 
+	/** Changes the URL data of a game feature plugin. Useful to change any options data that is parsed from the URL such as settings flags */
+	void UpdateGameFeaturePluginURL(const FString& NewPluginURL);
+	void UpdateGameFeaturePluginURL(const FString& NewPluginURL, const FGameFeaturePluginUpdateURLComplete& CompleteDelegate);
+
 	/** Gets the Install_Percent for single game feature plugin if it is active. */
 	bool GetGameFeaturePluginInstallPercent(const FString& PluginURL, float& Install_Percent) const;
 
@@ -260,14 +412,21 @@ public:
 	void UnloadGameFeaturePlugin(const FString& PluginURL, bool bKeepRegistered = false);
 	void UnloadGameFeaturePlugin(const FString& PluginURL, const FGameFeaturePluginUnloadComplete& CompleteDelegate, bool bKeepRegistered = false);
 
-	/** Uninstall the specified game feature plugin. Will remove the game feature plugin from the device if it was downloaded */
+	/** Releases any game data stored for this GameFeaturePlugin. Does not uninstall data and it will remain on disk. */
+	void ReleaseGameFeaturePlugin(const FString& PluginURL);
+	void ReleaseGameFeaturePlugin(const FString& PluginURL, const FGameFeaturePluginReleaseComplete& CompleteDelegate);
+
+	/** Uninstalls any game data stored for this GameFeaturePlugin and terminates the GameFeaturePlugin.
+		If the given PluginURL is not found this will create a GameFeaturePlugin first and attempt to run it through the uninstall flow.
+		This allows for the uninstalling of data that was installed on previous runs of the application where we haven't yet requested the
+		GameFeaturePlugin that we would like to uninstall data for on this run. */
 	void UninstallGameFeaturePlugin(const FString& PluginURL);
 	void UninstallGameFeaturePlugin(const FString& PluginURL, const FGameFeaturePluginUninstallComplete& CompleteDelegate);
 
-	/** Same as UninstallGameFeaturePlugin, but completely removes all tracking data associated with the plugin. */
+	/** Terminate the GameFeaturePlugin and remove all associated plugin tracking data. */
 	void TerminateGameFeaturePlugin(const FString& PluginURL);
-	void TerminateGameFeaturePlugin(const FString& PluginURL, const FGameFeaturePluginUninstallComplete& CompleteDelegate);
-
+	void TerminateGameFeaturePlugin(const FString& PluginURL, const FGameFeaturePluginTerminateComplete& CompleteDelegate);
+	
 	/** Attempt to cancel any state change. Calls back when cancelation is complete. Any other pending callbacks will be called with a canceled error. */
 	void CancelGameFeatureStateChange(const FString& PluginURL);
 	void CancelGameFeatureStateChange(const FString& PluginURL, const FGameFeaturePluginChangeStateComplete& CompleteDelegate);
@@ -348,6 +507,10 @@ private:
 	void OnGameFeatureLoading(const UGameFeatureData* GameFeatureData, const FString& PluginURL);
 	friend struct FGameFeaturePluginState_Loading;
 
+	void OnGameFeaturePauseChange(const FString& PluginURL, const FString& PluginName, FGameFeaturePauseStateChangeContext& Context);
+	friend struct FGameFeaturePluginState_Downloading;
+	friend struct FGameFeaturePluginState_Deactivating;
+
 	void OnAssetManagerCreated();
 
 	/** Scans for assets specified in the game feature data */
@@ -356,6 +519,8 @@ private:
 	static void RemoveGameFeatureFromAssetManager(const UGameFeatureData* GameFeatureToRemove, const FString& PluginName, const TArray<FName>& AddedPrimaryAssetTypes);
 
 private:
+	bool ShouldUpdatePluginURLData(const FString& NewPluginURL);
+
 	const UGameFeatureData* GetDataForStateMachine(UGameFeaturePluginStateMachine* GFSM) const;
 	const UGameFeatureData* GetRegisteredDataForStateMachine(UGameFeaturePluginStateMachine* GFSM) const;
 
@@ -399,18 +564,20 @@ private:
 		Unregistering,
 		Loading,
 		Activating,
-		Deactivating
+		Deactivating,
+		PauseChanged,
+		Count
 	};
 
 	void CallbackObservers(EObserverCallback CallbackType, const FString& PluginURL, 
 		const FString* PluginName = nullptr, 
 		const UGameFeatureData* GameFeatureData = nullptr, 
-		FGameFeatureDeactivatingContext* DeactivatingContext = nullptr);
+		FGameFeatureStateChangeContext* StateChangeContext = nullptr);
 
 private:
 	/** The list of all game feature plugin state machine objects */
 	UPROPERTY(Transient)
-	TMap<FString, UGameFeaturePluginStateMachine*> GameFeaturePluginStateMachines;
+	TMap<FGameFeaturePluginIdentifier, UGameFeaturePluginStateMachine*> GameFeaturePluginStateMachines;
 
 	/** Game feature plugin state machine objects that are being terminated. Used to prevent GC until termination is complete. */
 	UPROPERTY(Transient)
diff --git a/Engine/Plugins/Experimental/GameFeatures/Source/GameFeaturesEditor/Private/SGameFeatureStateWidget.cpp b/Engine/Plugins/Experimental/GameFeatures/Source/GameFeaturesEditor/Private/SGameFeatureStateWidget.cpp
--- a/Engine/Plugins/Experimental/GameFeatures/Source/GameFeaturesEditor/Private/SGameFeatureStateWidget.cpp
+++ b/Engine/Plugins/Experimental/GameFeatures/Source/GameFeaturesEditor/Private/SGameFeatureStateWidget.cpp
@@ -68,7 +68,7 @@ FText SGameFeatureStateWidget::GetDisplayNameOfState(EGameFeaturePluginState Sta
 
 FText SGameFeatureStateWidget::GetTooltipOfState(EGameFeaturePluginState State)
 {
-	static_assert((int32)EGameFeaturePluginState::MAX == 26, "");
+	static_assert((int32)EGameFeaturePluginState::MAX == 28, "");
 
 	switch (State)
 	{
@@ -78,6 +78,10 @@ FText SGameFeatureStateWidget::GetTooltipOfState(EGameFeaturePluginState State)
 		return LOCTEXT("StateTooltip_Terminal", "Final State before removal of the state machine");
 	case EGameFeaturePluginState::UnknownStatus:
 		return LOCTEXT("StateTooltip_UnknownStatus", "Initialized, but the only thing known is the URL to query status.");
+	case EGameFeaturePluginState::Uninstalling:
+		return LOCTEXT("StateTooltip_Uninstalling", "Transition state between StatusKnown -> Terminal for any plugin that can have data that needs to have local datat uninstalled.");
+	case EGameFeaturePluginState::ErrorUninstalling:
+		return LOCTEXT("StateTooltip_Error Uninstalling", "Error state for Uninstalling -> Terminal transition.");
 	case EGameFeaturePluginState::CheckingStatus:
 		return LOCTEXT("StateTooltip_CheckingStatus", "Transition state UnknownStatus -> StatusKnown. The status is in the process of being queried.");
 	case EGameFeaturePluginState::ErrorCheckingStatus:
@@ -86,10 +90,10 @@ FText SGameFeatureStateWidget::GetTooltipOfState(EGameFeaturePluginState State)
 		return LOCTEXT("StateTooltip_ErrorUnavailable", "Error state for UnknownStatus -> StatusKnown transition.");
 	case EGameFeaturePluginState::StatusKnown:
 		return LOCTEXT("StateTooltip_StatusKnown", "The plugin's information is known, but no action has taken place yet.");
-	case EGameFeaturePluginState::ErrorInstalling:
-		return LOCTEXT("StateTooltip_ErrorInstalling", "Error state for Installed -> StatusKnown and StatusKnown -> Installed transitions.");
-	case EGameFeaturePluginState::Uninstalling:
-		return LOCTEXT("StateTooltip_Uninstalling", "Transition state Installed -> StatusKnown. In the process of removing from local storage.");
+	case EGameFeaturePluginState::Releasing:
+		return LOCTEXT("StateTooltip_Releasing", "Transition State for Installed -> StatusKnown. Releases local data from any relevant caches.");
+	case EGameFeaturePluginState::ErrorManagingData:
+		return LOCTEXT("StateTooltip_ErrorManagingData", "Error state for Installed -> StatusKnown and StatusKnown -> Installed transitions.");
 	case EGameFeaturePluginState::Downloading:
 		return LOCTEXT("StateTooltip_Downloading", "Transition state StatusKnown -> Installed. In the process of adding to local storage.");
 	case EGameFeaturePluginState::Installed:
