In-editor lackluste preview on level load for. Needed for LD preview and HLOD generation.

diff --git a/Engine/Plugins/Experimental/GameFeatures/Source/GameFeatures/Private/GameFeaturePluginStateMachine.cpp b/Engine/Plugins/Experimental/GameFeatures/Source/GameFeatures/Private/GameFeaturePluginStateMachine.cpp
--- a/Engine/Plugins/Experimental/GameFeatures/Source/GameFeatures/Private/GameFeaturePluginStateMachine.cpp
+++ b/Engine/Plugins/Experimental/GameFeatures/Source/GameFeatures/Private/GameFeaturePluginStateMachine.cpp
@@ -2427,9 +2427,11 @@ UGameFeatureData* UGameFeaturePluginStateMachine::GetGameFeatureDataForActivePlu
 	return nullptr;
 }
 
-UGameFeatureData* UGameFeaturePluginStateMachine::GetGameFeatureDataForRegisteredPlugin()
+UGameFeatureData* UGameFeaturePluginStateMachine::GetGameFeatureDataForRegisteredPlugin(bool bCheckForRegistering /*= false*/)
 {
-	if (GetCurrentState() >= EGameFeaturePluginState::Registered)
+	const EGameFeaturePluginState CurrentState = GetCurrentState();
+
+	if (CurrentState >= EGameFeaturePluginState::Registered || (bCheckForRegistering && (CurrentState == EGameFeaturePluginState::Registering)))
 	{
 		return StateProperties.GameFeatureData;
 	}
diff --git a/Engine/Plugins/Experimental/GameFeatures/Source/GameFeatures/Private/GameFeaturePluginStateMachine.h b/Engine/Plugins/Experimental/GameFeatures/Source/GameFeatures/Private/GameFeaturePluginStateMachine.h
--- a/Engine/Plugins/Experimental/GameFeatures/Source/GameFeatures/Private/GameFeaturePluginStateMachine.h
+++ b/Engine/Plugins/Experimental/GameFeatures/Source/GameFeatures/Private/GameFeaturePluginStateMachine.h
@@ -480,7 +480,7 @@ public:
 	UGameFeatureData* GetGameFeatureDataForActivePlugin();
 
 	/** If the plugin is registered already, we will retrieve its game feature data */
-	UGameFeatureData* GetGameFeatureDataForRegisteredPlugin();
+	UGameFeatureData* GetGameFeatureDataForRegisteredPlugin(bool bCheckForRegistering = false);
 
 private:
 	/** Returns true if the specified state is not a transition state */
diff --git a/Engine/Plugins/Experimental/GameFeatures/Source/GameFeatures/Private/GameFeaturesSubsystem.cpp b/Engine/Plugins/Experimental/GameFeatures/Source/GameFeatures/Private/GameFeaturesSubsystem.cpp
--- a/Engine/Plugins/Experimental/GameFeatures/Source/GameFeatures/Private/GameFeaturesSubsystem.cpp
+++ b/Engine/Plugins/Experimental/GameFeatures/Source/GameFeatures/Private/GameFeaturesSubsystem.cpp
@@ -670,11 +670,11 @@ const UGameFeatureData* UGameFeaturesSubsystem::GetGameFeatureDataForActivePlugi
 	return nullptr;
 }
 
-const UGameFeatureData* UGameFeaturesSubsystem::GetGameFeatureDataForRegisteredPluginByURL(const FString& PluginURL)
+const UGameFeatureData* UGameFeaturesSubsystem::GetGameFeatureDataForRegisteredPluginByURL(const FString& PluginURL, bool bCheckForRegistering /*= false*/)
 {
 	if (UGameFeaturePluginStateMachine* GFSM = FindGameFeaturePluginStateMachine(PluginURL))
 	{
-		return GFSM->GetGameFeatureDataForRegisteredPlugin();
+		return GFSM->GetGameFeatureDataForRegisteredPlugin(bCheckForRegistering);
 	}
 
 	return nullptr;
@@ -689,11 +689,13 @@ bool UGameFeaturesSubsystem::IsGameFeaturePluginInstalled(const FString& PluginU
 	return false;
 }
 
-bool UGameFeaturesSubsystem::IsGameFeaturePluginRegistered(const FString& PluginURL) const
+bool UGameFeaturesSubsystem::IsGameFeaturePluginRegistered(const FString& PluginURL, bool bCheckForRegistering /*= false*/) const
 {
 	if (const UGameFeaturePluginStateMachine* StateMachine = FindGameFeaturePluginStateMachine(PluginURL))
 	{
-		return StateMachine->GetCurrentState() >= EGameFeaturePluginState::Registered;
+		const EGameFeaturePluginState CurrentState = StateMachine->GetCurrentState();
+
+		return StateMachine->GetCurrentState() >= EGameFeaturePluginState::Registered || (bCheckForRegistering && CurrentState == EGameFeaturePluginState::Registering);
 	}
 	return false;
 }
diff --git a/Engine/Plugins/Experimental/GameFeatures/Source/GameFeatures/Public/GameFeaturesSubsystem.h b/Engine/Plugins/Experimental/GameFeatures/Source/GameFeatures/Public/GameFeaturesSubsystem.h
--- a/Engine/Plugins/Experimental/GameFeatures/Source/GameFeatures/Public/GameFeaturesSubsystem.h
+++ b/Engine/Plugins/Experimental/GameFeatures/Source/GameFeatures/Public/GameFeaturesSubsystem.h
@@ -374,13 +374,13 @@ public:
 	const UGameFeatureData* GetGameFeatureDataForActivePluginByURL(const FString& PluginURL);
 
 	/** Returns the game feature data for a registered plugin specified by PluginURL */
-	const UGameFeatureData* GetGameFeatureDataForRegisteredPluginByURL(const FString& PluginURL);
+	const UGameFeatureData* GetGameFeatureDataForRegisteredPluginByURL(const FString& PluginURL, bool bCheckForRegistering = false);
 
 	/** Determines if a plugin is in the Installed state (or beyond) */
 	bool IsGameFeaturePluginInstalled(const FString& PluginURL) const;
 
 	/** Determines if a plugin is in the Registered state (or beyond) */
-	bool IsGameFeaturePluginRegistered(const FString& PluginURL) const;
+	bool IsGameFeaturePluginRegistered(const FString& PluginURL, bool bCheckForRegistering = false) const;
 
 	/** Determines if a plugin is in the Loaded state (or beyond) */
 	bool IsGameFeaturePluginLoaded(const FString& PluginURL) const;
