Fix installbundle plugin being reloaded when applying hotfixes which should only involve builtin plugins

diff --git a/Engine/Plugins/Experimental/GameFeatures/Source/GameFeatures/Private/GameFeaturesSubsystem.cpp b/Engine/Plugins/Experimental/GameFeatures/Source/GameFeatures/Private/GameFeaturesSubsystem.cpp
--- a/Engine/Plugins/Experimental/GameFeatures/Source/GameFeatures/Private/GameFeaturesSubsystem.cpp
+++ b/Engine/Plugins/Experimental/GameFeatures/Source/GameFeatures/Private/GameFeaturesSubsystem.cpp
@@ -533,7 +533,7 @@ EGameFeaturePluginProtocol UGameFeaturesSubsystem::GetPluginURLProtocol(FStringV
 {
 	for (EGameFeaturePluginProtocol Protocol : TEnumRange<EGameFeaturePluginProtocol>())
 	{
-		if (PluginURL.StartsWith(UE::GameFeatures::GameFeaturePluginProtocolPrefix(Protocol)))
+		if (UGameFeaturesSubsystem::IsPluginURLProtocol(PluginURL, Protocol))
 		{
 			return Protocol;
 		}
@@ -541,6 +541,11 @@ EGameFeaturePluginProtocol UGameFeaturesSubsystem::GetPluginURLProtocol(FStringV
 	return EGameFeaturePluginProtocol::Unknown;
 }
 
+bool UGameFeaturesSubsystem::IsPluginURLProtocol(FStringView PluginURL, EGameFeaturePluginProtocol PluginProtocol)
+{
+	return PluginURL.StartsWith(UE::GameFeatures::GameFeaturePluginProtocolPrefix(PluginProtocol));
+}
+
 void UGameFeaturesSubsystem::OnGameFeatureTerminating(const FString& PluginName, const FString& PluginURL)
 {
 	CallbackObservers(EObserverCallback::Terminating, PluginURL, &PluginName);
@@ -962,7 +967,7 @@ void UGameFeaturesSubsystem::UninstallGameFeaturePlugin(const FString& PluginURL
 	FString PluginURLForTerminate = PluginURL;
 
 	//InstallBundle Protocol GameFeatures may need to change their metadata to force this uninstall
-	if (EGameFeaturePluginProtocol::InstallBundle == GetPluginURLProtocol(PluginURL))
+	if (UGameFeaturesSubsystem::IsPluginURLProtocol(PluginURL, EGameFeaturePluginProtocol::InstallBundle))
 	{
 		// Parse a duplicate version of our current Metadata from the URL
 		FInstallBundlePluginProtocolMetaData ProtocolMetadata;
@@ -1052,8 +1057,18 @@ void UGameFeaturesSubsystem::LoadBuiltInGameFeaturePlugin(const TSharedRef<IPlug
 	// Make sure you are in a game feature plugins folder. All GameFeaturePlugins are rooted in a GameFeatures folder.
 	if (!PluginDescriptorFilename.IsEmpty() && GetDefault<UGameFeaturesSubsystemSettings>()->IsValidGameFeaturePlugin(FPaths::ConvertRelativePathToFull(PluginDescriptorFilename)) && FPaths::FileExists(PluginDescriptorFilename))
 	{
-		const FString PluginURL = GetPluginURL_FileProtocol(PluginDescriptorFilename);
-		if (GameSpecificPolicies->IsPluginAllowed(PluginURL))
+		FString PluginURL;
+		bool bIsFileProtocol = true;
+		if (GetPluginURLByName(Plugin->GetName(), PluginURL))
+		{
+			bIsFileProtocol = UGameFeaturesSubsystem::IsPluginURLProtocol(PluginURL, EGameFeaturePluginProtocol::File);
+		}
+		else
+		{
+			PluginURL = GetPluginURL_FileProtocol(PluginDescriptorFilename);
+		}
+
+		if (bIsFileProtocol && GameSpecificPolicies->IsPluginAllowed(PluginURL))
 		{
 			FGameFeaturePluginDetails PluginDetails;
 			if (GetGameFeaturePluginDetails(PluginDescriptorFilename, PluginDetails))
diff --git a/Engine/Plugins/Experimental/GameFeatures/Source/GameFeatures/Public/GameFeaturesSubsystem.h b/Engine/Plugins/Experimental/GameFeatures/Source/GameFeatures/Public/GameFeaturesSubsystem.h
--- a/Engine/Plugins/Experimental/GameFeatures/Source/GameFeatures/Public/GameFeaturesSubsystem.h
+++ b/Engine/Plugins/Experimental/GameFeatures/Source/GameFeatures/Public/GameFeaturesSubsystem.h
@@ -364,8 +364,12 @@ public:
 	static FString GetPluginURL_InstallBundleProtocol(const FString& PluginName, FName BundleName);
 	static FString GetPluginURL_InstallBundleProtocol(const FString& PluginName, const FInstallBundlePluginProtocolMetaData& ProtocolMetadata);
 
+	/** Returns the plugin protocol for the specified URL */
 	static EGameFeaturePluginProtocol GetPluginURLProtocol(FStringView PluginURL);
 
+	/** Tests whether the plugin URL is the specified protocol */
+	static bool IsPluginURLProtocol(FStringView PluginURL, EGameFeaturePluginProtocol PluginProtocol);
+
 public:
 	/** Returns all the active plugins GameFeatureDatas */
 	void GetGameFeatureDataForActivePlugins(TArray<const UGameFeatureData*>& OutActivePluginFeatureDatas);
