FNames containing asset paths are deprecated. FSoftObjectPath or FTopLevelAssetPath should be used instead.
Fixups for game features.

diff --git a/Engine/Plugins/Experimental/GameFeatures/Source/GameFeatures/Private/GameFeatureAction_AddComponents.cpp b/Engine/Plugins/Experimental/GameFeatures/Source/GameFeatures/Private/GameFeatureAction_AddComponents.cpp
--- a/Engine/Plugins/Experimental/GameFeatures/Source/GameFeatures/Private/GameFeatureAction_AddComponents.cpp
+++ b/Engine/Plugins/Experimental/GameFeatures/Source/GameFeatures/Private/GameFeatureAction_AddComponents.cpp
@@ -51,11 +51,11 @@ void UGameFeatureAction_AddComponents::AddAdditionalAssetBundleData(FAssetBundle
 		{
 			if (Entry.bClientComponent)
 			{
-				AssetBundleData.AddBundleAsset(UGameFeaturesSubsystemSettings::LoadStateClient, Entry.ComponentClass.ToSoftObjectPath());
+				AssetBundleData.AddBundleAsset(UGameFeaturesSubsystemSettings::LoadStateClient, Entry.ComponentClass.ToSoftObjectPath().GetAssetPath());
 			}
 			if (Entry.bServerComponent)
 			{
-				AssetBundleData.AddBundleAsset(UGameFeaturesSubsystemSettings::LoadStateServer, Entry.ComponentClass.ToSoftObjectPath());
+				AssetBundleData.AddBundleAsset(UGameFeaturesSubsystemSettings::LoadStateServer, Entry.ComponentClass.ToSoftObjectPath().GetAssetPath());
 			}
 		}
 	}
diff --git a/Engine/Plugins/Experimental/GameFeatures/Source/GameFeatures/Private/GameFeatureAction_DataRegistry.cpp b/Engine/Plugins/Experimental/GameFeatures/Source/GameFeatures/Private/GameFeatureAction_DataRegistry.cpp
--- a/Engine/Plugins/Experimental/GameFeatures/Source/GameFeatures/Private/GameFeatureAction_DataRegistry.cpp
+++ b/Engine/Plugins/Experimental/GameFeatures/Source/GameFeatures/Private/GameFeatureAction_DataRegistry.cpp
@@ -147,7 +147,7 @@ void UGameFeatureAction_DataRegistry::AddAdditionalAssetBundleData(FAssetBundleD
 	{
 		if(!RegistryToAdd.IsNull())
 		{
-			const FSoftObjectPath RegistryPath = RegistryToAdd.ToSoftObjectPath();
+			const FTopLevelAssetPath RegistryPath = RegistryToAdd.ToSoftObjectPath().GetAssetPath();
 
 			// Add for both clients and servers, this will not work properly for games that do not set those bundle states
 			// @TODO: If another way to preload specific assets is added, switch to that so it works regardless of bundles
diff --git a/Engine/Plugins/Experimental/GameFeatures/Source/GameFeatures/Private/GameFeatureAction_DataRegistrySource.cpp b/Engine/Plugins/Experimental/GameFeatures/Source/GameFeatures/Private/GameFeatureAction_DataRegistrySource.cpp
--- a/Engine/Plugins/Experimental/GameFeatures/Source/GameFeatures/Private/GameFeatureAction_DataRegistrySource.cpp
+++ b/Engine/Plugins/Experimental/GameFeatures/Source/GameFeatures/Private/GameFeatureAction_DataRegistrySource.cpp
@@ -195,7 +195,7 @@ void UGameFeatureAction_DataRegistrySource::AddAdditionalAssetBundleData(FAssetB
 
 		if (!RegistrySource.DataTableToAdd.IsNull())
 		{
-			const FSoftObjectPath DataTableSourcePath = RegistrySource.DataTableToAdd.ToSoftObjectPath();
+			const FTopLevelAssetPath DataTableSourcePath = RegistrySource.DataTableToAdd.ToSoftObjectPath().GetAssetPath();
 			if (RegistrySource.bClientSource)
 			{
 				AssetBundleData.AddBundleAsset(UGameFeaturesSubsystemSettings::LoadStateClient, DataTableSourcePath);
@@ -208,7 +208,7 @@ void UGameFeatureAction_DataRegistrySource::AddAdditionalAssetBundleData(FAssetB
 
 		if (!RegistrySource.CurveTableToAdd.IsNull())
 		{
-			const FSoftObjectPath CurveTableSourcePath = RegistrySource.CurveTableToAdd.ToSoftObjectPath();
+			const FTopLevelAssetPath CurveTableSourcePath = RegistrySource.CurveTableToAdd.ToSoftObjectPath().GetAssetPath();
 			if (RegistrySource.bClientSource)
 			{
 				AssetBundleData.AddBundleAsset(UGameFeaturesSubsystemSettings::LoadStateClient, CurveTableSourcePath);
diff --git a/Engine/Plugins/Experimental/GameFeatures/Source/GameFeatures/Private/GameFeaturesSubsystem.cpp b/Engine/Plugins/Experimental/GameFeatures/Source/GameFeatures/Private/GameFeaturesSubsystem.cpp
--- a/Engine/Plugins/Experimental/GameFeatures/Source/GameFeatures/Private/GameFeaturesSubsystem.cpp
+++ b/Engine/Plugins/Experimental/GameFeatures/Source/GameFeatures/Private/GameFeaturesSubsystem.cpp
@@ -1632,7 +1632,7 @@ void UGameFeaturesSubsystem::FilterInactivePluginAssets(TArray<FAssetData>& Asse
 {
 	AssetsToFilter.RemoveAllSwap([ActivePluginNames = GetActivePluginNames()](const FAssetData& Asset) 
 	{
-		return !GameFeaturesSubsystem::IsContentWithinActivePlugin(Asset.ObjectPath.ToString(), ActivePluginNames);
+		return !GameFeaturesSubsystem::IsContentWithinActivePlugin(Asset.GetObjectPathString(), ActivePluginNames);
 	});
 }
 EBuiltInAutoState UGameFeaturesSubsystem::DetermineBuiltInInitialFeatureState(TSharedPtr<FJsonObject> Descriptor, const FString& ErrorContext)
