Cleanup delegate moves in UGameFeaturesSubsystem::ChangeGameFeatureDestination().  These we wrong in that they could be called after being moved from, but delegates can't be moved right now anyway so it didn't acutally cause an issue. I'm just removing them for correctness.

diff --git a/Engine/Plugins/Experimental/GameFeatures/Source/GameFeatures/Private/GameFeaturesSubsystem.cpp b/Engine/Plugins/Experimental/GameFeatures/Source/GameFeatures/Private/GameFeaturesSubsystem.cpp
--- a/Engine/Plugins/Experimental/GameFeatures/Source/GameFeatures/Private/GameFeaturesSubsystem.cpp
+++ b/Engine/Plugins/Experimental/GameFeatures/Source/GameFeatures/Private/GameFeaturesSubsystem.cpp
@@ -1356,7 +1356,7 @@ void UGameFeaturesSubsystem::LoadBuiltInGameFeaturePluginComplete(const UE::Game
 void UGameFeaturesSubsystem::ChangeGameFeatureDestination(UGameFeaturePluginStateMachine* Machine, const FGameFeaturePluginStateRange& StateRange, FGameFeaturePluginChangeStateComplete CompleteDelegate)
 {
 	const bool bSetDestination = Machine->SetDestination(StateRange,
-		FGameFeatureStateTransitionComplete::CreateUObject(this, &ThisClass::ChangeGameFeatureTargetStateComplete, MoveTemp(CompleteDelegate)));
+		FGameFeatureStateTransitionComplete::CreateUObject(this, &ThisClass::ChangeGameFeatureTargetStateComplete, CompleteDelegate));
 
 	if (bSetDestination)
 	{
@@ -1365,7 +1365,7 @@ void UGameFeaturesSubsystem::ChangeGameFeatureDestination(UGameFeaturePluginStat
 	else
 	{
 		// Try canceling any current transition, then retry
-		auto OnCanceled = [this, StateRange, CompleteDelegate=MoveTemp(CompleteDelegate)](UGameFeaturePluginStateMachine* Machine) mutable
+		auto OnCanceled = [this, StateRange, CompleteDelegate](UGameFeaturePluginStateMachine* Machine) mutable
 		{
 			// Special case for terminal state since it cannot be exited, we need to make a new machine
 			if (Machine->GetCurrentState() == EGameFeaturePluginState::Terminal)
@@ -1377,7 +1377,7 @@ void UGameFeaturesSubsystem::ChangeGameFeatureDestination(UGameFeaturePluginStat
 
 			// Now that the transition has been canceled, retry reaching the desired destination
 			const bool bSetDestination = Machine->SetDestination(StateRange,
-				FGameFeatureStateTransitionComplete::CreateUObject(this, &ThisClass::ChangeGameFeatureTargetStateComplete, MoveTemp(CompleteDelegate)));
+				FGameFeatureStateTransitionComplete::CreateUObject(this, &ThisClass::ChangeGameFeatureTargetStateComplete, CompleteDelegate));
 
 			if (!ensure(bSetDestination))
 			{
