Add a GameFeatureAction to inject a content bundle

diff --git a/Engine/Plugins/Experimental/GameFeatures/Source/GameFeatures/Private/GameFeatureAction_AddWPContent.cpp b/Engine/Plugins/Experimental/GameFeatures/Source/GameFeatures/Private/GameFeatureAction_AddWPContent.cpp
new file mode 100644
--- /dev/null
+++ b/Engine/Plugins/Experimental/GameFeatures/Source/GameFeatures/Private/GameFeatureAction_AddWPContent.cpp
@@ -0,0 +1,63 @@
+// Copyright Epic Games, Inc. All Rights Reserved.
+
+#include "GameFeatureAction_AddWPContent.h"
+
+#include "WorldPartition/ContentBundle/ContentBundleDescriptor.h"
+#include "WorldPartition/ContentBundle/ContentBundleClient.h"
+#include "GameFeatureData.h"
+#include "GameFeaturesSubsystem.h"
+
+UGameFeatureAction_AddWPContent::UGameFeatureAction_AddWPContent(const FObjectInitializer& ObjectInitializer)
+	: Super(ObjectInitializer)
+{
+	ContentBundleDescriptor = ObjectInitializer.CreateDefaultSubobject<UContentBundleDescriptor>(this, TEXT("ContentBundleDescriptor"));
+
+#if WITH_EDITOR
+	if (UGameFeatureData* GameFeatureData = GetTypedOuter<UGameFeatureData>())
+	{
+		ContentBundleDescriptor->SetDisplayName(GameFeatureData->GetName());
+
+		FString PackagePath = GetPackage()->GetName();
+		FName MountPoint = FPackageName::GetPackageMountPoint(PackagePath);
+		ContentBundleDescriptor->SetPackageRoot(MountPoint.ToString());
+	}
+#endif
+}
+
+void UGameFeatureAction_AddWPContent::OnGameFeatureRegistering()
+{
+	Super::OnGameFeatureRegistering();
+	ContentBundleClient = FContentBundleClient::CreateClient(ContentBundleDescriptor, GetTypedOuter<UGameFeatureData>()->GetName());
+	UE_CLOG(ContentBundleClient == nullptr, LogGameFeatures, Error, TEXT("OnGameFeatureRegistering %s: Failed to create a content bundle client for %s"), *GetPathName(), *ContentBundleDescriptor->GetDisplayName())
+}
+
+void UGameFeatureAction_AddWPContent::OnGameFeatureUnregistering()
+{
+	if (ContentBundleClient != nullptr)
+	{
+		ContentBundleClient->RequestUnregister();
+		ContentBundleClient = nullptr;
+	}
+
+	Super::OnGameFeatureUnregistering();
+}
+
+void UGameFeatureAction_AddWPContent::OnGameFeatureActivating()
+{
+	Super::OnGameFeatureActivating();
+	
+	if (ContentBundleClient)
+	{
+		ContentBundleClient->RequestContentInjection();
+	}
+}
+
+void UGameFeatureAction_AddWPContent::OnGameFeatureDeactivating(FGameFeatureDeactivatingContext& Context)
+{
+	if (ContentBundleClient)
+	{
+		ContentBundleClient->RequestRemoveContent();
+	}
+
+	Super::OnGameFeatureDeactivating(Context);
+}
\ No newline at end of file
diff --git a/Engine/Plugins/Experimental/GameFeatures/Source/GameFeatures/Public/GameFeatureAction_AddWPContent.h b/Engine/Plugins/Experimental/GameFeatures/Source/GameFeatures/Public/GameFeatureAction_AddWPContent.h
new file mode 100644
--- /dev/null
+++ b/Engine/Plugins/Experimental/GameFeatures/Source/GameFeatures/Public/GameFeatureAction_AddWPContent.h
@@ -0,0 +1,38 @@
+// Copyright Epic Games, Inc. All Rights Reserved.
+
+#pragma once
+
+#include "CoreMinimal.h"
+#include "GameFeatureAction.h"
+
+#include "GameFeatureAction_AddWPContent.generated.h"
+
+class UGameFeatureData;
+class IPlugin;
+class FContentBundleClient;
+class UContentBundleDescriptor;
+
+/**
+ *
+ */
+UCLASS(meta = (DisplayName = "Add World Partition Content"))
+class UGameFeatureAction_AddWPContent : public UGameFeatureAction
+{
+	GENERATED_UCLASS_BODY()
+
+public:
+	//~ Begin UGameFeatureAction interface
+	virtual void OnGameFeatureRegistering() override;
+	virtual void OnGameFeatureUnregistering() override;
+	virtual void OnGameFeatureActivating() override;
+	virtual void OnGameFeatureDeactivating(FGameFeatureDeactivatingContext& Context) override;
+	//~ End UGameFeatureAction interface
+
+	UGameFeatureData* GetGameFeatureData() const;
+
+private:
+	UPROPERTY(EditAnywhere, Category = ContentBundle)
+	TObjectPtr<UContentBundleDescriptor> ContentBundleDescriptor;
+
+	TSharedPtr<FContentBundleClient> ContentBundleClient;
+};
\ No newline at end of file
