Reintroduce GFP ensure

diff --git a/Engine/Plugins/Experimental/GameFeatures/Source/GameFeatures/Private/GameFeaturePluginStateMachine.cpp b/Engine/Plugins/Experimental/GameFeatures/Source/GameFeatures/Private/GameFeaturePluginStateMachine.cpp
--- a/Engine/Plugins/Experimental/GameFeatures/Source/GameFeatures/Private/GameFeaturePluginStateMachine.cpp
+++ b/Engine/Plugins/Experimental/GameFeatures/Source/GameFeatures/Private/GameFeaturePluginStateMachine.cpp
@@ -87,7 +87,7 @@ namespace UE::GameFeatures
 					FReferenceChainSearch(AssetObj, EReferenceChainSearchMode::Shortest | EReferenceChainSearchMode::PrintResults);
 				}
 
-				//ensureAlwaysMsgf(false, TEXT("GFP %s failed to unload asset %s!"), *PluginName, *AssetData.GetFullName());
+				ensureAlwaysMsgf(false, TEXT("GFP %s failed to unload asset %s!"), *PluginName, *AssetData.GetFullName());
 			}
 		};
 
