Add some additional logging around GFP transition failures and transition canceling.

diff --git a/Engine/Plugins/Experimental/GameFeatures/Source/GameFeatures/Private/GameFeaturesSubsystem.cpp b/Engine/Plugins/Experimental/GameFeatures/Source/GameFeatures/Private/GameFeaturesSubsystem.cpp
--- a/Engine/Plugins/Experimental/GameFeatures/Source/GameFeatures/Private/GameFeaturesSubsystem.cpp
+++ b/Engine/Plugins/Experimental/GameFeatures/Source/GameFeatures/Private/GameFeaturesSubsystem.cpp
@@ -767,6 +767,8 @@ void UGameFeaturesSubsystem::ChangeGameFeatureTargetState(const FString& PluginU
 		StateMachine = FindGameFeaturePluginStateMachine(PluginURL);
 		if (!StateMachine)
 		{
+			UE_LOG(LogGameFeatures, Log, TEXT("Cannot create GFP State Machine: Plugin not allowed %s"), *PluginURL);
+
 			CompleteDelegate.ExecuteIfBound(UE::GameFeatures::FResult(MakeError(UE::GameFeatures::SubsystemErrorNamespace + UE::GameFeatures::CommonErrorCodes::PluginNotAllowed)));
 			return;
 		}
@@ -782,6 +784,8 @@ void UGameFeaturesSubsystem::ChangeGameFeatureTargetState(const FString& PluginU
 	{
 		if (TargetPluginState > StateMachine->GetCurrentState() || TargetPluginState > StateMachine->GetDestination())
 		{
+			UE_LOG(LogGameFeatures, Log, TEXT("Cannot change game feature target state: Plugin not allowed %s"), *PluginURL);
+
 			CompleteDelegate.ExecuteIfBound(UE::GameFeatures::FResult(MakeError(UE::GameFeatures::SubsystemErrorNamespace + UE::GameFeatures::CommonErrorCodes::PluginNotAllowed)));
 			return;
 		}
@@ -1364,6 +1368,12 @@ void UGameFeaturesSubsystem::ChangeGameFeatureDestination(UGameFeaturePluginStat
 	}
 	else
 	{
+		FGameFeaturePluginStateRange CurrDesitination = Machine->GetDestination();
+		UE_LOG(LogGameFeatures, Display, TEXT("ChangeGameFeatureDestination: Attempting to cancel transition for Game Feature %s. Desired [%s, %s]. Current [%s, %s]"), 
+			*Machine->GetGameFeatureName(), 
+			*UE::GameFeatures::ToString(StateRange.MinState), *UE::GameFeatures::ToString(StateRange.MaxState),
+			*UE::GameFeatures::ToString(CurrDesitination.MinState), *UE::GameFeatures::ToString(CurrDesitination.MaxState));
+
 		// Try canceling any current transition, then retry
 		auto OnCanceled = [this, StateRange, CompleteDelegate](UGameFeaturePluginStateMachine* Machine) mutable
 		{
