Add an option to the game feature plugin creation code to allow turning on bEnabledByDefault if the game setup wants it
Add read-only display of EnabledByDefault and ExplicitlyLoaded so users can see what's going on
Change Lyra to enable ExplicitlyLoaded but it leave it off by default to avoid backward compatibility issues with changing the base config files

diff --git a/Engine/Plugins/Editor/PluginBrowser/Source/PluginBrowser/Private/PluginMetadataObject.cpp b/Engine/Plugins/Editor/PluginBrowser/Source/PluginBrowser/Private/PluginMetadataObject.cpp
--- a/Engine/Plugins/Editor/PluginBrowser/Source/PluginBrowser/Private/PluginMetadataObject.cpp
+++ b/Engine/Plugins/Editor/PluginBrowser/Source/PluginBrowser/Private/PluginMetadataObject.cpp
@@ -51,6 +51,8 @@ void UPluginMetadataObject::PopulateFromPlugin(TSharedPtr<IPlugin> InPlugin)
 	SupportURL = InDescriptor.SupportURL;
 	bCanContainContent = InDescriptor.bCanContainContent;
 	bIsBetaVersion = InDescriptor.bIsBetaVersion;
+	bIsEnabledByDefault = (InDescriptor.EnabledByDefault == EPluginEnabledByDefault::Enabled);
+	bExplicitlyLoaded = InDescriptor.bExplicitlyLoaded;
 
 	Plugins.Reset(InDescriptor.Plugins.Num());
 	for (const FPluginReferenceDescriptor& PluginRefDesc : InDescriptor.Plugins)
diff --git a/Engine/Plugins/Editor/PluginBrowser/Source/PluginBrowser/Private/PluginMetadataObject.h b/Engine/Plugins/Editor/PluginBrowser/Source/PluginBrowser/Private/PluginMetadataObject.h
--- a/Engine/Plugins/Editor/PluginBrowser/Source/PluginBrowser/Private/PluginMetadataObject.h
+++ b/Engine/Plugins/Editor/PluginBrowser/Source/PluginBrowser/Private/PluginMetadataObject.h
@@ -112,6 +112,14 @@ public:
 	UPROPERTY(EditAnywhere, Category = Details)
 	bool bIsBetaVersion;
 
+	/** If true, the plugin will be enabled by default without needing to modify the project file. This cannot be changed while the editor is running */
+	UPROPERTY(VisibleAnywhere, Category = Details)
+	bool bIsEnabledByDefault;
+
+	/** If true, the plugin will not load at startup and must be explicitly loaded. This cannot be changed while the editor is running */
+	UPROPERTY(VisibleAnywhere, Category = Details)
+	bool bExplicitlyLoaded;
+
 	/** Plugins used by this plugin */
 	UPROPERTY(EditAnywhere, Category = Dependencies, meta=(TitleProperty=Name))
 	TArray<FPluginReferenceMetadata> Plugins;
diff --git a/Engine/Plugins/Experimental/GameFeatures/Source/GameFeaturesEditor/Private/GameFeaturesEditorModule.cpp b/Engine/Plugins/Experimental/GameFeatures/Source/GameFeaturesEditor/Private/GameFeaturesEditorModule.cpp
--- a/Engine/Plugins/Experimental/GameFeatures/Source/GameFeaturesEditor/Private/GameFeaturesEditorModule.cpp
+++ b/Engine/Plugins/Experimental/GameFeatures/Source/GameFeaturesEditor/Private/GameFeaturesEditorModule.cpp
@@ -32,13 +32,15 @@
 
 struct FGameFeaturePluginTemplateDescription : public FPluginTemplateDescription
 {
-	FGameFeaturePluginTemplateDescription(FText InName, FText InDescription, FString InOnDiskPath, TSubclassOf<UGameFeatureData> GameFeatureDataClassOverride, FString GameFeatureDataNameOverride)
+	FGameFeaturePluginTemplateDescription(FText InName, FText InDescription, FString InOnDiskPath
+		, TSubclassOf<UGameFeatureData> GameFeatureDataClassOverride, FString GameFeatureDataNameOverride, EPluginEnabledByDefault InEnabledByDefault)
 		: FPluginTemplateDescription(InName, InDescription, InOnDiskPath, /*bCanContainContent=*/ true, EHostType::Runtime)
 	{
 		SortPriority = 10;
 		bCanBePlacedInEngine = false;
 		GameFeatureDataName = !GameFeatureDataNameOverride.IsEmpty() ? GameFeatureDataNameOverride : FString();
 		GameFeatureDataClass = GameFeatureDataClassOverride != nullptr ? GameFeatureDataClassOverride : TSubclassOf<UGameFeatureData>(UGameFeatureData::StaticClass());
+		PluginEnabledByDefault = InEnabledByDefault;
 	}
 
 	virtual bool ValidatePathForPlugin(const FString& ProposedAbsolutePluginPath, FText& OutErrorMessage) override
@@ -73,8 +75,8 @@ struct FGameFeaturePluginTemplateDescription : public FPluginTemplateDescription
 		Descriptor.AdditionalFieldsToWrite.FindOrAdd(TEXT("BuiltInInitialFeatureState")) = MakeShared<FJsonValueString>(TEXT("Active"));
 		Descriptor.Category = TEXT("Game Features");
 
-		// Game features should be disabled by default, as they are managed by the game feature subsystem
-		Descriptor.EnabledByDefault = EPluginEnabledByDefault::Disabled;
+		// Game features should not be enabled by default if the game wants to strictly manage default settings in the target settings
+		Descriptor.EnabledByDefault = PluginEnabledByDefault;
 
 		if (Descriptor.Modules.Num() > 0)
 		{
@@ -136,6 +138,7 @@ struct FGameFeaturePluginTemplateDescription : public FPluginTemplateDescription
 
 	TSubclassOf<UGameFeatureData> GameFeatureDataClass;
 	FString GameFeatureDataName;
+	EPluginEnabledByDefault PluginEnabledByDefault = EPluginEnabledByDefault::Disabled;
 };
 
 //////////////////////////////////////////////////////////////////////
@@ -226,7 +229,8 @@ class FGameFeaturesEditorModule : public FDefaultModuleImpl
 					PluginTemplate.Description,
 					PluginTemplate.Path.Path,
 					PluginTemplate.DefaultGameFeatureDataClass,
-					PluginTemplate.DefaultGameFeatureDataName)));
+					PluginTemplate.DefaultGameFeatureDataName,
+					PluginTemplate.bIsEnabledByDefault ? EPluginEnabledByDefault::Enabled : EPluginEnabledByDefault::Disabled)));
 			}
 		}
 	}
diff --git a/Engine/Plugins/Experimental/GameFeatures/Source/GameFeaturesEditor/Private/GameFeaturesEditorSettings.h b/Engine/Plugins/Experimental/GameFeatures/Source/GameFeaturesEditor/Private/GameFeaturesEditorSettings.h
--- a/Engine/Plugins/Experimental/GameFeatures/Source/GameFeaturesEditor/Private/GameFeaturesEditorSettings.h
+++ b/Engine/Plugins/Experimental/GameFeatures/Source/GameFeaturesEditor/Private/GameFeaturesEditorSettings.h
@@ -40,6 +40,11 @@ struct FPluginTemplateData
 	/** The default name of the created game feature data assets. If empty, will use the plugin name. */
 	UPROPERTY(config, EditAnywhere, Category = Plugins)
 	FString DefaultGameFeatureDataName;
+
+	/** If true, the created plugin will be enabled by default without needing to be added to the project file. */
+	UPROPERTY(config, EditAnywhere, Category = Plugins)
+	bool bIsEnabledByDefault = false;
+
 };
 
 
