Change Game Feature plugins to load DefaultPluginName.ini using the same method as normal plugins so + works on arrays
The code to load PluginName.ini still runs with a log message

diff --git a/Engine/Plugins/Experimental/GameFeatures/Source/GameFeatures/Private/GameFeatureData.cpp b/Engine/Plugins/Experimental/GameFeatures/Source/GameFeatures/Private/GameFeatureData.cpp
--- a/Engine/Plugins/Experimental/GameFeatures/Source/GameFeatures/Private/GameFeatureData.cpp
+++ b/Engine/Plugins/Experimental/GameFeatures/Source/GameFeatures/Private/GameFeatureData.cpp
@@ -5,6 +5,7 @@
 #include "GameFeaturesSubsystem.h"
 #include "GameFeaturesSubsystemSettings.h"
 #include "Misc/ConfigCacheIni.h"
+#include "Misc/ConfigContext.h"
 #include "UObject/UObjectHash.h"
 #include "UObject/CoreRedirects.h"
 
@@ -65,20 +66,32 @@ void UGameFeatureData::InitializeBasePluginIniFile(const FString& PluginInstalle
 	const bool bForceReloadFromDisk = false;
 	const bool bWriteDestIni = false;
 
-	// look for plugin ini
-	// @note: We use the generated config dir, because ReloadConfig will use this path + plugin name to read it out of GConfig, so when we add it to GConfig we make
-	// sure the string is this format
-	FString PluginConfigFilename = FString::Printf(TEXT("%s%s/%s.ini"), *FPaths::GeneratedConfigDir(), ANSI_TO_TCHAR(FPlatformProperties::PlatformName()), *PluginName);
-	FPaths::MakeStandardFilename(PluginConfigFilename);
+	// This will be the generated path including platform
+	FString PluginConfigFilename = GConfig->GetConfigFilename(*PluginName);
 
+	// Try the deprecated path first that doesn't include the Default prefix
 	FConfigFile& PluginConfig = GConfig->Add(PluginConfigFilename, FConfigFile());
 	if (!FConfigCacheIni::LoadExternalIniFile(PluginConfig, *PluginName, *EngineConfigDir, *PluginConfigDir, bIsBaseIniName, nullptr, bForceReloadFromDisk, bWriteDestIni))
 	{
-		// Nothing to add, remove from map
-		GConfig->Remove(PluginConfigFilename);
+		// Now try the same rules as PluginManager using Default and the config hierarchy
+		FConfigContext Context = FConfigContext::ReadIntoPluginFile(PluginConfig, *FPaths::GetPath(PluginInstalledFilename));
+
+		if (!Context.Load(*PluginName))
+		{
+			// Nothing to add, remove from map
+			GConfig->Remove(PluginConfigFilename);
+		}
+		else
+		{
+			FCoreRedirects::ReadRedirectsFromIni(PluginConfigFilename);
+			ReloadConfigs(PluginConfig);
+		}
 	}
 	else
 	{
+		// This is the deprecated loading path that doesn't handle cases like + in arrays
+		UE_LOG(LogGameFeatures, Log, TEXT("[GameFeatureData %s]: Loaded deprecated config %s, rename to start with Default for normal parsing"), *GetPathNameSafe(this), *PluginConfigFilename);
+
 		FCoreRedirects::ReadRedirectsFromIni(PluginConfigFilename);
 		ReloadConfigs(PluginConfig);
 	}
