Fix an issue where optimization can force an Engine module reference in <redacted>.

diff --git a/Engine/Plugins/Experimental/GameFeatures/Source/GameFeatures/Public/GameFeaturesSubsystem.h b/Engine/Plugins/Experimental/GameFeatures/Source/GameFeatures/Public/GameFeaturesSubsystem.h
--- a/Engine/Plugins/Experimental/GameFeatures/Source/GameFeatures/Public/GameFeaturesSubsystem.h
+++ b/Engine/Plugins/Experimental/GameFeatures/Source/GameFeatures/Public/GameFeaturesSubsystem.h
@@ -305,7 +305,7 @@ public:
 	virtual void Deinitialize() override;
 	//~End of UEngineSubsystem interface
 
-	static UGameFeaturesSubsystem& Get() { return *GEngine->GetEngineSubsystem<UGameFeaturesSubsystem>(); }
+	static FORCENOINLINE UGameFeaturesSubsystem& Get() { return *GEngine->GetEngineSubsystem<UGameFeaturesSubsystem>(); }
 
 public:
 	/** Loads the specified game feature data and its bundles */
