Fix an issue where engine was not depended on when depending on gamefeaturePlugins.
Only depending on Engine when not running as a program.

diff --git a/Engine/Plugins/Experimental/GameFeatures/Source/GameFeatures/Public/GameFeaturesSubsystem.h b/Engine/Plugins/Experimental/GameFeatures/Source/GameFeatures/Public/GameFeaturesSubsystem.h
--- a/Engine/Plugins/Experimental/GameFeatures/Source/GameFeatures/Public/GameFeaturesSubsystem.h
+++ b/Engine/Plugins/Experimental/GameFeatures/Source/GameFeatures/Public/GameFeaturesSubsystem.h
@@ -305,7 +305,7 @@ public:
 	virtual void Deinitialize() override;
 	//~End of UEngineSubsystem interface
 
-	static FORCENOINLINE UGameFeaturesSubsystem& Get() { return *GEngine->GetEngineSubsystem<UGameFeaturesSubsystem>(); }
+	static UGameFeaturesSubsystem& Get() { return *GEngine->GetEngineSubsystem<UGameFeaturesSubsystem>(); }
 
 public:
 	/** Loads the specified game feature data and its bundles */
