Add ExplicitLoadGameFeaturePlugin

diff --git a/Engine/Plugins/Experimental/GameFeatures/Source/GameFeatures/Private/GameFeaturesProjectPolicies.cpp b/Engine/Plugins/Experimental/GameFeatures/Source/GameFeatures/Private/GameFeaturesProjectPolicies.cpp
--- a/Engine/Plugins/Experimental/GameFeatures/Source/GameFeatures/Private/GameFeaturesProjectPolicies.cpp
+++ b/Engine/Plugins/Experimental/GameFeatures/Source/GameFeatures/Private/GameFeaturesProjectPolicies.cpp
@@ -41,3 +41,15 @@ const TArray<FName> UDefaultGameFeaturesProjectPolicies::GetPreloadBundleStateFo
 	}
 	return FeatureBundles;
 }
+
+void UGameFeaturesProjectPolicies::ExplicitLoadGameFeaturePlugin(const FString& PluginURL, const FGameFeaturePluginLoadComplete& CompleteDelegate, const bool bActivateGameFeatures)
+{
+	if (bActivateGameFeatures)
+	{
+		UGameFeaturesSubsystem::Get().LoadAndActivateGameFeaturePlugin(PluginURL, CompleteDelegate);
+	}
+	else
+	{
+		UGameFeaturesSubsystem::Get().LoadGameFeaturePlugin(PluginURL, CompleteDelegate);
+	}
+}
\ No newline at end of file
diff --git a/Engine/Plugins/Experimental/GameFeatures/Source/GameFeatures/Public/GameFeaturesProjectPolicies.h b/Engine/Plugins/Experimental/GameFeatures/Source/GameFeatures/Public/GameFeaturesProjectPolicies.h
--- a/Engine/Plugins/Experimental/GameFeatures/Source/GameFeatures/Public/GameFeaturesProjectPolicies.h
+++ b/Engine/Plugins/Experimental/GameFeatures/Source/GameFeatures/Public/GameFeaturesProjectPolicies.h
@@ -4,6 +4,7 @@
 
 #include "CoreMinimal.h"
 #include "UObject/PrimaryAssetId.h"
+#include "GameFeaturesSubsystem.h"
 
 #include "GameFeaturesProjectPolicies.generated.h"
 
@@ -37,6 +38,10 @@ public:
 	// Called to determine if a plugin is allowed to be loaded or not
 	// (e.g., when doing a fast cook a game might want to disable some or all game feature plugins)
 	virtual bool IsPluginAllowed(const FString& PluginURL) const { return true; }
+
+	// Called by code that explicitly wants to load a specific plugin
+	// (e.g., when using a fast cook a game might want to allow explicitly loaded game feature plugins)
+	virtual void ExplicitLoadGameFeaturePlugin(const FString& PluginURL, const FGameFeaturePluginLoadComplete& CompleteDelegate, const bool bActivateGameFeatures);
 };
 
 // This is a default implementation that immediately processes all game feature plugins the based on
