Allow for not removing dependencies in bundle manager RequestRemoveContent through new EInstallBundleReleaseRequestFlags::ExplicitRemoveList
Add support for this and other Release Flags into InstallBundle protocol GameFeaturePlugins
#tests none

diff --git a/Engine/Plugins/Experimental/DefaultInstallBundleManager/Source/Private/DefaultInstallBundleManager.cpp b/Engine/Plugins/Experimental/DefaultInstallBundleManager/Source/Private/DefaultInstallBundleManager.cpp
--- a/Engine/Plugins/Experimental/DefaultInstallBundleManager/Source/Private/DefaultInstallBundleManager.cpp
+++ b/Engine/Plugins/Experimental/DefaultInstallBundleManager/Source/Private/DefaultInstallBundleManager.cpp
@@ -3291,8 +3291,17 @@ TValueOrError<FInstallBundleRequestInfo, EInstallBundleResult> FDefaultInstallBu
 
 	LOG_INSTALL_BUNDLE_MAN_OVERRIDE(LogVerbosityOverride, Display, TEXT("RequestReleaseContent"));
 
-	TSet<FName> BundlesToRelease = GatherBundlesForRequest(ReleaseNames, RetInfo.InfoFlags);
 	TSet<FName> BundlesToKeep = GatherBundlesForRequest(KeepNames, RetInfo.InfoFlags);
+	
+	TSet<FName> BundlesToRelease;
+	if (EnumHasAnyFlags(Flags, EInstallBundleReleaseRequestFlags::ExplicitRemoveList))
+	{
+		BundlesToRelease.Append(ReleaseNames);
+	}
+	else
+	{
+		BundlesToRelease = GatherBundlesForRequest(ReleaseNames, RetInfo.InfoFlags);
+	}
 
 	// Don't release shared dependencies
 	for (const FName& BundleName : BundlesToKeep)
diff --git a/Engine/Plugins/Experimental/GameFeatures/Source/GameFeatures/Private/GameFeaturePluginStateMachine.cpp b/Engine/Plugins/Experimental/GameFeatures/Source/GameFeatures/Private/GameFeaturePluginStateMachine.cpp
--- a/Engine/Plugins/Experimental/GameFeatures/Source/GameFeatures/Private/GameFeaturePluginStateMachine.cpp
+++ b/Engine/Plugins/Experimental/GameFeatures/Source/GameFeatures/Private/GameFeaturePluginStateMachine.cpp
@@ -230,6 +230,7 @@ namespace UE::GameFeatures
 	XOPTION(Bundles,					TEXT("Bundles"))					\
 	XOPTION(Version,					TEXT("V"))							\
 	XOPTION(Flags,						TEXT("Flags"))						\
+	XOPTION(ReleaseFlags,				TEXT("ReleaseFlags"))						\
 	XOPTION(UninstallBeforeTerminate,	TEXT("UninstallBeforeTerminate"))	\
 	XOPTION(UserPauseDownload,			TEXT("Paused"))
 
@@ -821,13 +822,15 @@ struct FBaseDataReleaseGameFeaturePluginState : public FGameFeaturePluginState
 			}
 		}
 	}
+	/** Determine what kind of release request flags we submit */
+	virtual EInstallBundleReleaseRequestFlags GetReleaseRequestFlags() const
+	{
+		return StateProperties.ProtocolMetadata.GetSubtype<FInstallBundlePluginProtocolMetaData>().ReleaseInstallBundleFlags;
+	}
 
 	/** Determines what state you transition to in the event of a success or failure to release content */
 	virtual EGameFeaturePluginState GetSuccessTransitionState() const = 0;
 	virtual EGameFeaturePluginState GetFailureTransitionState() const = 0;
-
-	/** Determine what kind of release request flags we submit */
-	virtual EInstallBundleReleaseRequestFlags GetReleaseRequestFlags() const = 0;
 };
 
 struct FGameFeaturePluginState_Uninstalling : public FBaseDataReleaseGameFeaturePluginState
@@ -849,7 +852,8 @@ struct FGameFeaturePluginState_Uninstalling : public FBaseDataReleaseGameFeature
 	//Must be overridden to determine what kind of release request we submit
 	virtual EInstallBundleReleaseRequestFlags GetReleaseRequestFlags() const override
 	{
-		return EInstallBundleReleaseRequestFlags::RemoveFilesIfPossible;
+		const EInstallBundleReleaseRequestFlags BaseFlags = FBaseDataReleaseGameFeaturePluginState::GetReleaseRequestFlags();
+		return (BaseFlags | EInstallBundleReleaseRequestFlags::RemoveFilesIfPossible);
 	}
 
 	virtual bool TryUpdateURLData(const FString& NewPluginURL) override
@@ -894,12 +898,6 @@ struct FGameFeaturePluginState_Releasing : public FBaseDataReleaseGameFeaturePlu
 	{
 		return EGameFeaturePluginState::ErrorManagingData;
 	}
-
-	//Must be overridden to determine what kind of release request we submit
-	virtual EInstallBundleReleaseRequestFlags GetReleaseRequestFlags() const override
-	{
-		return EInstallBundleReleaseRequestFlags::None;
-	}
 };
 
 struct FGameFeaturePluginState_Downloading : public FGameFeaturePluginState
@@ -2650,6 +2648,11 @@ FString FInstallBundlePluginProtocolMetaData::ToString() const
 		ReturnedString.Append(FString::Printf(TEXT("%s%s%s%d"), UE::GameFeatures::PluginURLStructureInfo::OptionSeperator, *LexToString(EGameFeatureInstallBundleProtocolOptions::Flags), UE::GameFeatures::PluginURLStructureInfo::OptionAssignOperator, static_cast<uint32>(InstallBundleFlags)));
 	}
 
+	if (ReleaseInstallBundleFlags != FDefaultValues::Default_ReleaseInstallBundleFlags)
+	{
+		ReturnedString.Append(FString::Printf(TEXT("%s%s%s%d"), UE::GameFeatures::PluginURLStructureInfo::OptionSeperator, *LexToString(EGameFeatureInstallBundleProtocolOptions::ReleaseFlags), UE::GameFeatures::PluginURLStructureInfo::OptionAssignOperator, static_cast<uint32>(ReleaseInstallBundleFlags)));
+	}
+
 	//Encode bUninstallBeforeTerminate if not the default
 	if (bUninstallBeforeTerminate != FDefaultValues::Default_bUninstallBeforeTerminate)
 	{
@@ -2662,7 +2665,7 @@ FString FInstallBundlePluginProtocolMetaData::ToString() const
 		ReturnedString.Append(FString::Printf(TEXT("%s%s%s%s"), UE::GameFeatures::PluginURLStructureInfo::OptionSeperator, *LexToString(EGameFeatureInstallBundleProtocolOptions::UserPauseDownload), UE::GameFeatures::PluginURLStructureInfo::OptionAssignOperator, *LexToString(bUserPauseDownload)));
 	}
 
-	static_assert(static_cast<uint8>(EGameFeatureInstallBundleProtocolOptions::Count) == 5, "Update this function to handle the newly added EGameFeatureInstallBundleProtocolOptions value!");
+	static_assert(static_cast<uint8>(EGameFeatureInstallBundleProtocolOptions::Count) == 6, "Update this function to handle the newly added EGameFeatureInstallBundleProtocolOptions value!");
 
 	return ReturnedString;
 }
@@ -2750,6 +2753,22 @@ bool FInstallBundlePluginProtocolMetaData::FromString(const FString& URLString,
 						break;
 					}
 
+					case EGameFeatureInstallBundleProtocolOptions::ReleaseFlags:
+					{
+						if (OptionStrings.Num() != 2)
+						{
+							bParseSuccess = false;
+							UE_LOG(LogGameFeatures, Warning, TEXT("Error parsing InstallBundle protocol options URL %s . Invalid ReleaseFlags Option!"), *URLString);
+						}
+						//Expect to parse uint32 flag version of this as a 2nd parameter
+						else
+						{
+							Metadata.ReleaseInstallBundleFlags= static_cast<EInstallBundleReleaseRequestFlags>(FCString::Strtoi(*OptionStrings[1], nullptr, 10));
+						}
+
+						break;
+					}
+
 					case EGameFeatureInstallBundleProtocolOptions::UninstallBeforeTerminate:
 					{
 						if (OptionStrings.Num() != 2)
@@ -2799,7 +2818,7 @@ bool FInstallBundlePluginProtocolMetaData::FromString(const FString& URLString,
 		UE_LOG(LogGameFeatures, Error, TEXT("Error parsing InstallBundle protocol options URL %s . No Bundle List Found!"), *URLString);
 	}
 
-	static_assert(static_cast<uint8>(EGameFeatureInstallBundleProtocolOptions::Count) == 5, "Update this function to handle the newly added EGameFeatureInstallBundleProtocolOptions value!");
+	static_assert(static_cast<uint8>(EGameFeatureInstallBundleProtocolOptions::Count) == 6, "Update this function to handle the newly added EGameFeatureInstallBundleProtocolOptions value!");
 
 	return bParseSuccess;
 }
diff --git a/Engine/Plugins/Experimental/GameFeatures/Source/GameFeatures/Private/GameFeaturesSubsystem.cpp b/Engine/Plugins/Experimental/GameFeatures/Source/GameFeatures/Private/GameFeaturesSubsystem.cpp
--- a/Engine/Plugins/Experimental/GameFeatures/Source/GameFeatures/Private/GameFeaturesSubsystem.cpp
+++ b/Engine/Plugins/Experimental/GameFeatures/Source/GameFeatures/Private/GameFeaturesSubsystem.cpp
@@ -26,6 +26,7 @@ const uint32 FInstallBundlePluginProtocolMetaData::FDefaultValues::CurrentVersio
 const bool FInstallBundlePluginProtocolMetaData::FDefaultValues::Default_bUninstallBeforeTerminate = true;
 const bool FInstallBundlePluginProtocolMetaData::FDefaultValues::Default_bUserPauseDownload = false;
 const EInstallBundleRequestFlags FInstallBundlePluginProtocolMetaData::FDefaultValues::Default_InstallBundleFlags = EInstallBundleRequestFlags::Defaults;
+const EInstallBundleReleaseRequestFlags FInstallBundlePluginProtocolMetaData::FDefaultValues::Default_ReleaseInstallBundleFlags = EInstallBundleReleaseRequestFlags::None;
 
 namespace UE::GameFeatures
 {
diff --git a/Engine/Plugins/Experimental/GameFeatures/Source/GameFeatures/Public/GameFeaturesSubsystem.h b/Engine/Plugins/Experimental/GameFeatures/Source/GameFeatures/Public/GameFeaturesSubsystem.h
--- a/Engine/Plugins/Experimental/GameFeatures/Source/GameFeatures/Public/GameFeaturesSubsystem.h
+++ b/Engine/Plugins/Experimental/GameFeatures/Source/GameFeatures/Public/GameFeaturesSubsystem.h
@@ -21,6 +21,7 @@ class FJsonObject;
 struct FWorldContext;
 struct FGameFeaturePluginStateRange;
 enum class EInstallBundleRequestFlags : uint32;
+enum class EInstallBundleReleaseRequestFlags : uint32;
 
 /** Holds static global information about how our PluginURLs are structured */
 namespace UE::GameFeatures
@@ -269,6 +270,9 @@ struct GAMEFEATURES_API FInstallBundlePluginProtocolMetaData
 	/** EInstallBundleRequestFlags utilized during the download/install by InstallBundleManager */
 	EInstallBundleRequestFlags InstallBundleFlags;
 
+	/** EInstallBundleReleaseRequestFlags utilized during our release and uninstall states */
+	EInstallBundleReleaseRequestFlags ReleaseInstallBundleFlags;
+
 	/** If we want to set the Downloading state to pause because of user interaction */
 	bool bUserPauseDownload;
 
@@ -290,6 +294,7 @@ private:
 		static const bool Default_bUninstallBeforeTerminate;
 		static const bool Default_bUserPauseDownload;
 		static const EInstallBundleRequestFlags Default_InstallBundleFlags;
+		static const EInstallBundleReleaseRequestFlags Default_ReleaseInstallBundleFlags;
 	};
 };
 
diff --git a/Engine/Source/Runtime/InstallBundleManager/Public/InstallBundleTypes.h b/Engine/Source/Runtime/InstallBundleManager/Public/InstallBundleTypes.h
--- a/Engine/Source/Runtime/InstallBundleManager/Public/InstallBundleTypes.h
+++ b/Engine/Source/Runtime/InstallBundleManager/Public/InstallBundleTypes.h
@@ -168,6 +168,7 @@ enum class EInstallBundleReleaseRequestFlags : uint32
 {
 	None = 0,
 	RemoveFilesIfPossible = (1 << 0),  // Bundle sources must support removal, and bundle must not be part of the source's cache
+	ExplicitRemoveList = (1 << 1),	   // Only attempt to remove explicitly supplied bundles instead of automatically removing dependencies 
 };
 ENUM_CLASS_FLAGS(EInstallBundleReleaseRequestFlags)
 
