Fixing GameFeaturePlugins attempting to update plugin URLs when the GameFeaturePlugin is already terminal

diff --git a/Engine/Plugins/Experimental/GameFeatures/Source/GameFeatures/Private/GameFeaturePluginStateMachine.cpp b/Engine/Plugins/Experimental/GameFeatures/Source/GameFeatures/Private/GameFeaturePluginStateMachine.cpp
--- a/Engine/Plugins/Experimental/GameFeatures/Source/GameFeatures/Private/GameFeaturePluginStateMachine.cpp
+++ b/Engine/Plugins/Experimental/GameFeatures/Source/GameFeatures/Private/GameFeaturePluginStateMachine.cpp
@@ -470,6 +470,12 @@ struct FGameFeaturePluginState_Terminal : public FDestinationGameFeaturePluginSt
 
 	bool bEnteredTerminalState = false;
 
+	virtual bool TryUpdateURLData(const FString& NewPluginURL) override
+	{
+		//Should never update our URL during Terminal
+		return false;
+	}
+
 	virtual void BeginState() override
 	{
 		checkf(!bEnteredTerminalState, TEXT("Plugin entered terminal state more than once! %s"), *StateProperties.PluginIdentifier.GetFullPluginURL());
diff --git a/Engine/Plugins/Experimental/GameFeatures/Source/GameFeatures/Private/GameFeaturesSubsystem.cpp b/Engine/Plugins/Experimental/GameFeatures/Source/GameFeatures/Private/GameFeaturesSubsystem.cpp
--- a/Engine/Plugins/Experimental/GameFeatures/Source/GameFeatures/Private/GameFeaturesSubsystem.cpp
+++ b/Engine/Plugins/Experimental/GameFeatures/Source/GameFeatures/Private/GameFeaturesSubsystem.cpp
@@ -859,7 +859,15 @@ bool UGameFeaturesSubsystem::ShouldUpdatePluginURLData(const FString& NewPluginU
 		return false;
 	}
 
+	//Should always be valid at this point
 	check(StateMachine);
+	
+	//Make sure our StateMachine isn't in terminal, don't want to update Terminal plugins
+	if (TerminalGameFeaturePluginStateMachines.Contains(StateMachine) || (StateMachine->GetCurrentState() == EGameFeaturePluginState::Terminal))
+	{
+		return false;
+	}
+
 	if (StateMachine->GetPluginURL().Equals(NewPluginURL, ESearchCase::IgnoreCase))
 	{
 		return false;
