Fix non-unity

#rnx

diff --git a/Engine/Plugins/Experimental/GameFeatures/Source/GameFeatures/Private/GameFeatureDataAssetDependencyGatherer.cpp b/Engine/Plugins/Experimental/GameFeatures/Source/GameFeatures/Private/GameFeatureDataAssetDependencyGatherer.cpp
--- a/Engine/Plugins/Experimental/GameFeatures/Source/GameFeatures/Private/GameFeatureDataAssetDependencyGatherer.cpp
+++ b/Engine/Plugins/Experimental/GameFeatures/Source/GameFeatures/Private/GameFeatureDataAssetDependencyGatherer.cpp
@@ -7,6 +7,7 @@
 #include "HAL/FileManager.h" 
 #include "WorldPartition/ContentBundle/ContentBundlePaths.h"
 #include "AssetRegistry/AssetData.h"
+#include "Engine/Level.h"
 
 // Register FGameFeatureDataAssetDependencyGatherer for UGameFeatureData class
 REGISTER_ASSETDEPENDENCY_GATHERER(FGameFeatureDataAssetDependencyGatherer, UGameFeatureData);
