Fix non-unity

#rnx

diff --git a/Engine/Plugins/Experimental/GameFeatures/Source/GameFeatures/Private/GameFeatureData.cpp b/Engine/Plugins/Experimental/GameFeatures/Source/GameFeatures/Private/GameFeatureData.cpp
--- a/Engine/Plugins/Experimental/GameFeatures/Source/GameFeatures/Private/GameFeatureData.cpp
+++ b/Engine/Plugins/Experimental/GameFeatures/Source/GameFeatures/Private/GameFeatureData.cpp
@@ -10,6 +10,7 @@
 #include "UObject/CoreRedirects.h"
 #include "GameFeatureAction_AddWPContent.h"
 #include "WorldPartition/ContentBundle/ContentBundleDescriptor.h"
+#include "AssetRegistry/AssetData.h"
 
 #include UE_INLINE_GENERATED_CPP_BY_NAME(GameFeatureData)
 
