[BugFix] reviewing GameFeatureData from CL 22843375 crashes

diff --git a/Engine/Plugins/Experimental/GameFeatures/Source/GameFeatures/Private/GameFeaturesSubsystem.cpp b/Engine/Plugins/Experimental/GameFeatures/Source/GameFeatures/Private/GameFeaturesSubsystem.cpp
--- a/Engine/Plugins/Experimental/GameFeatures/Source/GameFeatures/Private/GameFeaturesSubsystem.cpp
+++ b/Engine/Plugins/Experimental/GameFeatures/Source/GameFeatures/Private/GameFeaturesSubsystem.cpp
@@ -1316,10 +1316,13 @@ UGameFeaturePluginStateMachine* UGameFeaturesSubsystem::FindGameFeaturePluginSta
 
 UGameFeaturePluginStateMachine* UGameFeaturesSubsystem::FindGameFeaturePluginStateMachine(const FString& PluginURL) const
 {
-	TObjectPtr<UGameFeaturePluginStateMachine> const* ExistingStateMachine = GameFeaturePluginStateMachines.Find(PluginURL);
-	if (ExistingStateMachine)
+	if (!PluginURL.IsEmpty())
 	{
-		return *ExistingStateMachine;
+		TObjectPtr<UGameFeaturePluginStateMachine> const* ExistingStateMachine = GameFeaturePluginStateMachines.Find(PluginURL);
+		if (ExistingStateMachine)
+		{
+			return *ExistingStateMachine;
+		}
 	}
 
 	return nullptr;
