FGameFeaturePluginState_Unmounting
- Fix plugin not actually being removed from plugin manager. This would break subsequent reloads of different but same named plugins.

diff --git a/Engine/Plugins/Experimental/GameFeatures/Source/GameFeatures/Private/GameFeaturePluginStateMachine.cpp b/Engine/Plugins/Experimental/GameFeatures/Source/GameFeatures/Private/GameFeaturePluginStateMachine.cpp
--- a/Engine/Plugins/Experimental/GameFeatures/Source/GameFeatures/Private/GameFeaturePluginStateMachine.cpp
+++ b/Engine/Plugins/Experimental/GameFeatures/Source/GameFeatures/Private/GameFeaturePluginStateMachine.cpp
@@ -1341,7 +1341,7 @@ struct FGameFeaturePluginState_Unmounting : public FGameFeaturePluginState
 
 		if (StateProperties.bAddedPluginToManager)
 		{
-			verify(IPluginManager::Get().RemoveFromPluginsList(StateProperties.PluginName));
+			verify(IPluginManager::Get().RemoveFromPluginsList(StateProperties.PluginInstalledFilename));
 			StateProperties.bAddedPluginToManager = false;
 		}
 
