#UE Fix diffing for GameFeatureData assets

diff --git a/Engine/Plugins/Experimental/GameFeatures/Source/GameFeaturesEditor/Private/GameFeatureDataDetailsCustomization.cpp b/Engine/Plugins/Experimental/GameFeatures/Source/GameFeaturesEditor/Private/GameFeatureDataDetailsCustomization.cpp
--- a/Engine/Plugins/Experimental/GameFeatures/Source/GameFeaturesEditor/Private/GameFeatureDataDetailsCustomization.cpp
+++ b/Engine/Plugins/Experimental/GameFeatures/Source/GameFeaturesEditor/Private/GameFeatureDataDetailsCustomization.cpp
@@ -44,7 +44,7 @@ void FGameFeatureDataDetailsCustomization::CustomizeDetails(IDetailLayoutBuilder
 	ObjectsBeingCustomized.Empty();
 	DetailBuilder.GetObjectsBeingCustomized(/*out*/ ObjectsBeingCustomized);
 
-	if (ObjectsBeingCustomized.Num() == 1)
+	if (ObjectsBeingCustomized.Num() == 1 && !ObjectsBeingCustomized[0]->GetPackage()->HasAnyPackageFlags(PKG_ForDiffing))
 	{
 		const UGameFeatureData* GameFeature = CastChecked<const UGameFeatureData>(ObjectsBeingCustomized[0]);
 
