Add trace scopes related to GameFeaturePlugins

#rnx

diff --git a/Engine/Plugins/Experimental/GameFeatures/Source/GameFeatures/Private/GameFeaturePluginStateMachine.cpp b/Engine/Plugins/Experimental/GameFeatures/Source/GameFeatures/Private/GameFeaturePluginStateMachine.cpp
--- a/Engine/Plugins/Experimental/GameFeatures/Source/GameFeatures/Private/GameFeaturePluginStateMachine.cpp
+++ b/Engine/Plugins/Experimental/GameFeatures/Source/GameFeatures/Private/GameFeaturePluginStateMachine.cpp
@@ -17,6 +17,7 @@
 #include "Misc/FileHelper.h"
 #include "Misc/Paths.h"
 #include "Algo/AllOf.h"
+#include "ProfilingDebugging/CpuProfilerTrace.h"
 #include "Serialization/ArrayReader.h"
 #include "Serialization/MemoryReader.h"
 #include "Stats/Stats.h"
@@ -1522,6 +1523,7 @@ struct FGameFeaturePluginState_Mounting : public FGameFeaturePluginState
 
 	virtual void UpdateState(FGameFeaturePluginStateStatus& StateStatus) override
 	{
+		TRACE_CPUPROFILER_EVENT_SCOPE(GFP_Mounting);
 		if (!Result.HasValue())
 		{
 			StateStatus.SetTransitionError(EGameFeaturePluginState::ErrorMounting, Result);
@@ -1623,6 +1625,7 @@ struct FGameFeaturePluginState_WaitingForDependencies : public FGameFeaturePlugi
 
 	virtual void UpdateState(FGameFeaturePluginStateStatus& StateStatus) override
 	{
+		TRACE_CPUPROFILER_EVENT_SCOPE(GFP_WaitingForDependencies);
 		checkf(!StateProperties.PluginInstalledFilename.IsEmpty(), TEXT("PluginInstalledFilename must be set by the loading dependencies phase. PluginURL: %s"), *StateProperties.PluginIdentifier.GetFullPluginURL());
 		checkf(FPaths::GetExtension(StateProperties.PluginInstalledFilename) == TEXT("uplugin"), TEXT("PluginInstalledFilename must have a uplugin extension. PluginURL: %s"), *StateProperties.PluginIdentifier.GetFullPluginURL());
 
@@ -1824,6 +1827,7 @@ struct FGameFeaturePluginState_Registering : public FGameFeaturePluginState
 
 	virtual void UpdateState(FGameFeaturePluginStateStatus& StateStatus) override
 	{
+		TRACE_CPUPROFILER_EVENT_SCOPE(GFP_Registering);
 		const FString PluginFolder = FPaths::GetPath(StateProperties.PluginInstalledFilename);
 		UGameplayTagsManager::Get().AddTagIniSearchPath(PluginFolder / TEXT("Config") / TEXT("Tags"));
 
@@ -1975,6 +1979,7 @@ struct FGameFeaturePluginState_Loading : public FGameFeaturePluginState
 
 	virtual void UpdateState(FGameFeaturePluginStateStatus& StateStatus) override
 	{
+		TRACE_CPUPROFILER_EVENT_SCOPE(GFP_Loading);
 		check(StateProperties.GameFeatureData);
 
 		// AssetManager
@@ -2124,6 +2129,7 @@ struct FGameFeaturePluginState_Activating : public FGameFeaturePluginState
 
 	virtual void UpdateState(FGameFeaturePluginStateStatus& StateStatus) override
 	{
+		TRACE_CPUPROFILER_EVENT_SCOPE(GFP_Activating);
 		check(GEngine);
 		check(StateProperties.GameFeatureData);
 
diff --git a/Engine/Plugins/Experimental/GameFeatures/Source/GameFeatures/Private/GameFeaturesSubsystem.cpp b/Engine/Plugins/Experimental/GameFeatures/Source/GameFeatures/Private/GameFeaturesSubsystem.cpp
--- a/Engine/Plugins/Experimental/GameFeatures/Source/GameFeatures/Private/GameFeaturesSubsystem.cpp
+++ b/Engine/Plugins/Experimental/GameFeatures/Source/GameFeatures/Private/GameFeaturesSubsystem.cpp
@@ -12,6 +12,7 @@
 #include "Misc/Paths.h"
 #include "Misc/FileHelper.h"
 #include "Containers/Ticker.h"
+#include "ProfilingDebugging/CpuProfilerTrace.h"
 #include "Serialization/JsonReader.h"
 #include "Serialization/JsonSerializer.h"
 #include "Engine/AssetManager.h"
@@ -426,6 +427,7 @@ void UGameFeaturesSubsystem::UnloadGameFeatureData(const UGameFeatureData* GameF
 
 void UGameFeaturesSubsystem::AddGameFeatureToAssetManager(const UGameFeatureData* GameFeatureToAdd, const FString& PluginName, TArray<FName>& OutNewPrimaryAssetTypes)
 {
+	TRACE_CPUPROFILER_EVENT_SCOPE(GFP_AddToAssetManager);
 	check(GameFeatureToAdd);
 	FString PluginRootPath = TEXT("/") + PluginName + TEXT("/");
 	UAssetManager& LocalAssetManager = UAssetManager::Get();
@@ -1250,6 +1252,7 @@ EGameFeaturePluginState UGameFeaturesSubsystem::GetPluginState(FGameFeaturePlugi
 
 bool UGameFeaturesSubsystem::GetGameFeaturePluginDetails(const FString& PluginDescriptorFilename, FGameFeaturePluginDetails& OutPluginDetails) const
 {
+	TRACE_CPUPROFILER_EVENT_SCOPE(GFP_GetPluginDetails);
 	// @TODO: We load the descriptor 2-3 per plugin because FPluginReferenceDescriptor doesn't cache any of this info.
 	// GFPs are implemented with a plugin so FPluginReferenceDescriptor doesn't know anything about them.
 	// Need a better way of storing GFP specific plugin data...
@@ -1386,6 +1389,7 @@ UGameFeaturePluginStateMachine* UGameFeaturesSubsystem::FindGameFeaturePluginSta
 
 UGameFeaturePluginStateMachine* UGameFeaturesSubsystem::FindOrCreateGameFeaturePluginStateMachine(const FString& PluginURL)
 {
+	TRACE_CPUPROFILER_EVENT_SCOPE(GFP_FindOrCreateStateMachine);
 	FGameFeaturePluginIdentifier PluginIdentifier(PluginURL);
 	if (UGameFeaturePluginStateMachine* ExistingStateMachine = FindGameFeaturePluginStateMachine(PluginIdentifier))
 	{
diff --git a/Engine/Source/Runtime/Projects/Private/PluginManager.cpp b/Engine/Source/Runtime/Projects/Private/PluginManager.cpp
--- a/Engine/Source/Runtime/Projects/Private/PluginManager.cpp
+++ b/Engine/Source/Runtime/Projects/Private/PluginManager.cpp
@@ -26,6 +26,7 @@
 #include "Misc/ScopeRWLock.h"
 #include "Algo/Accumulate.h"
 #include "Containers/VersePath.h"
+#include "ProfilingDebugging/CpuProfilerTrace.h"
 #if READ_TARGET_ENABLED_PLUGINS_FROM_RECEIPT
 #include "TargetReceipt.h"
 #endif
@@ -2578,6 +2579,7 @@ bool FPluginManager::TryMountExplicitlyLoadedPluginVersion(TSharedRef<FPlugin>*
 
 void FPluginManager::MountPluginFromExternalSource(const TSharedRef<FPlugin>& Plugin)
 {
+	TRACE_CPUPROFILER_EVENT_SCOPE(MountPluginFromExternalSource);
 	if (GWarn)
 	{
 		GWarn->BeginSlowTask(FText::Format(LOCTEXT("MountingPluginFiles", "Mounting plugin {0}..."), FText::FromString(Plugin->GetFriendlyName())), /*ShowProgressDialog*/ true, /*bShowCancelButton*/ false);
@@ -2646,6 +2648,7 @@ bool FPluginManager::UnmountExplicitlyLoadedPlugin(const FString& PluginName, FT
 
 bool FPluginManager::UnmountPluginFromExternalSource(const TSharedPtr<FPlugin>& Plugin, FText* OutReason)
 {
+	TRACE_CPUPROFILER_EVENT_SCOPE(UnmountPluginFromExternalSource);
 	if (!Plugin.IsValid() || Plugin->bEnabled == false)
 	{
 		// Does not exist or is not loaded
