Additional error reporting for object leaks when unmounting GFPs.
Rename objects that leak so that new copies of them can be reloaded if the GFP plugin is mounted again.

diff --git a/Engine/Plugins/Experimental/GameFeatures/Source/GameFeatures/Private/GameFeaturePluginStateMachine.cpp b/Engine/Plugins/Experimental/GameFeatures/Source/GameFeatures/Private/GameFeaturePluginStateMachine.cpp
--- a/Engine/Plugins/Experimental/GameFeatures/Source/GameFeatures/Private/GameFeaturePluginStateMachine.cpp
+++ b/Engine/Plugins/Experimental/GameFeatures/Source/GameFeatures/Private/GameFeaturePluginStateMachine.cpp
@@ -25,7 +25,8 @@
 #include "GameFeaturesSubsystem.h"
 #include "GameFeaturesProjectPolicies.h"
 #include "Containers/Ticker.h"
-#include "UObject/ReferenceChainSearch.h"
+#include "UObject/ObjectRename.h"
+#include "UObject/UObjectAllocator.h"
 
 #include UE_INLINE_GENERATED_CPP_BY_NAME(GameFeaturePluginStateMachine)
 
@@ -44,15 +45,44 @@ namespace UE::GameFeatures
 		ShouldLogMountedFiles,
 		TEXT("Should the newly mounted files be logged."));
 
-	static TAutoConsoleVariable<bool> CVarVerifyPluginUnload(TEXT("GameFeaturePlugin.VerifyUnload"), 
-		true,
-		TEXT("Verify plugin assets are no longer in memory when unloading."),
-		ECVF_Default);
+	static bool GRenameLeakedPackages = true;
+	static FAutoConsoleVariableRef CVarRenameLeakedPackages(
+		TEXT("GameFeaturePlugin.LeakedAssetTrace.RenameLeakedPackages"),
+		GRenameLeakedPackages,
+		TEXT("Should packages which are leaked after the Game Feature Plugin is unloaded or unmounted."),
+		ECVF_Default
+	);
+
+	static int32 GLeakedAssetTrace_Severity = 2;
+	static FAutoConsoleVariableRef CVarLeakedAssetTrace_Severity(
+	#if UE_BUILD_SHIPPING
+		TEXT("GameFeaturePlugin.LeakedAssetTrace.Severity.Shipping"),
+	#else
+		TEXT("GameFeaturePlugin.LeakedAssetTrace.Severity"),
+	#endif
+		GLeakedAssetTrace_Severity,
+		TEXT("Controls severity of logging when the engine detects that assets from an Game Feature Plugin were leaked during unloading or unmounting.\n")
+		TEXT("0 - all reference tracing and logging is disabled\n")
+		TEXT("1 - logs an error\n")
+		TEXT("2 - ensure\n")
+		TEXT("3 - fatal error\n"),
+		ECVF_Default
+	);
+
+	static int32 GLeakedAssetTrace_TraceMode = (UE_BUILD_SHIPPING ? 0 : 1);
+	static FAutoConsoleVariableRef CVarLeakedAssetTrace_TraceMode(
+	#if UE_BUILD_SHIPPING
+		TEXT("GameFeaturePlugin.LeakedAssetTrace.TraceMode.Shipping"),
+	#else
+		TEXT("GameFeaturePlugin.LeakedASsetTrace.TraceMode"),
+	#endif
+		GLeakedAssetTrace_TraceMode,
+		TEXT("Controls detail level of reference tracing when the engine detects that assets from a Game Feature Plugin were leaked during unloading or unmounting.\n")
+		TEXT("0 - direct references only\n")
+		TEXT("1 - full reference trace"),
+		ECVF_Default
+	);
 
-	static TAutoConsoleVariable<bool> CVarVerifyPluginUnloadDumpChains(TEXT("GameFeaturePlugin.VerifyUnloadDumpChains"),
-		false,
-		TEXT("Dump reference chains for any detected plugin asset leaks."),
-		ECVF_Default);
 	#define GAME_FEATURE_PLUGIN_STATE_TO_STRING(inEnum, inText) case EGameFeaturePluginState::inEnum: return TEXT(#inEnum);
 	FString ToString(EGameFeaturePluginState InType)
 	{
@@ -66,68 +96,116 @@ namespace UE::GameFeatures
 	}
 	#undef GAME_FEATURE_PLUGIN_STATE_TO_STRING
 
-	// Verify that all assets from this plugin have been unloaded and GC'd	
-	void VerifyAssetsUnloaded(const FString& PluginName, bool bIgnoreGameFeatureData)
+	// Check if any assets from the plugin mount point have leaked, and if so trace them.
+	// Then rename them to allow new copies of them to be loaded. 
+	void HandlePossibleAssetLeaks(const FString& PluginName, UPackage* IgnorePackage = nullptr)
 	{
-#if (!UE_BUILD_SHIPPING || UE_SERVER || WITH_EDITOR)
-		if (!UE::GameFeatures::CVarVerifyPluginUnload.GetValueOnGameThread())
+		TRACE_CPUPROFILER_EVENT_SCOPE(HandlePossibleAssetLeaks);
+
+		const bool bFindLeakedPackages = GLeakedAssetTrace_Severity != 0 || GRenameLeakedPackages;
+		if (!bFindLeakedPackages)
 		{
 			return;
 		}
+		
+		TSet<UPackage*> LeakedPackages;
+		TStringBuilder<512> Prefix;
+		Prefix << '/' << PluginName << '/';
 
-		FARFilter PluginArFilter;
-		PluginArFilter.PackagePaths.Add(FName(TEXT("/") + PluginName));
-		PluginArFilter.bRecursivePaths = true;
-
-		auto CheckForLoadedAsset = [](const FString& PluginName, const FAssetData &AssetData)
+		TStringBuilder<NAME_SIZE> NameBuffer;
 		{
-			if (AssetData.IsAssetLoaded())
+			// If the UObject hash knew about package mount roots, we could avoid this loop
+			TRACE_CPUPROFILER_EVENT_SCOPE(PackageLoop);
+			FPermanentObjectPoolExtents PermExtents;
+			ForEachObjectOfClass(UPackage::StaticClass(), [&](UObject* Obj)
 			{
-				UE_LOG(LogGameFeatures, Error, TEXT("GFP %s failed to unload asset %s!"), *PluginName, *AssetData.GetFullName());
-
-				if (CVarVerifyPluginUnloadDumpChains.GetValueOnGameThread())
+				if (UPackage* Package = CastChecked<UPackage>(Obj, ECastCheckedType::NullAllowed))
 				{
-					UObject* AssetObj = AssetData.GetAsset();
-					//EInternalObjectFlags InternalFlags = AssetObj->GetInternalFlags();
-					FReferenceChainSearch(AssetObj, EReferenceChainSearchMode::Shortest | EReferenceChainSearchMode::PrintResults);
+					if (Package == IgnorePackage)
+					{
+						return;
+					}
+					if (PermExtents.Contains(Package))
+					{
+						return;
+					}
+					NameBuffer.Reset();
+					Package->GetFName().GetDisplayNameEntry()->AppendNameToString(NameBuffer);
+					if (NameBuffer.ToView().StartsWith(Prefix, ESearchCase::IgnoreCase))
+					{
+						LeakedPackages.Add(Package);
+					}
 				}
+			});
+		}
 
-				ensureAlwaysMsgf(false, TEXT("GFP %s failed to unload asset %s!"), *PluginName, *AssetData.GetFullName());
-			}
-		};
+		if (LeakedPackages.Num() == 0)
+		{
+			return;
+		}
 
-		if (bIgnoreGameFeatureData)
+		if (GLeakedAssetTrace_Severity != 0)
 		{
-			FARFilter RawGameFeatureDataFilter;
-			RawGameFeatureDataFilter.ClassPaths.Add(UGameFeatureData::StaticClass()->GetClassPathName());
-			RawGameFeatureDataFilter.bRecursiveClasses = true;
+			EPrintStaleReferencesOptions Options = EPrintStaleReferencesOptions::None;
+			switch (GLeakedAssetTrace_Severity)
+			{
+			case 3:
+				Options = EPrintStaleReferencesOptions::Fatal;
+				break;
+			case 2: 
+				Options = EPrintStaleReferencesOptions::Ensure | EPrintStaleReferencesOptions::Error;
+				break;
+			case 1:
+			default:
+				Options = EPrintStaleReferencesOptions::Error;
+				break;
+			}	
 
-			FARCompiledFilter GameFeatureDataFilter;
-			UAssetManager::Get().GetAssetRegistry().CompileFilter(RawGameFeatureDataFilter, GameFeatureDataFilter);
+			if (GLeakedAssetTrace_TraceMode == 0)
+			{
+				Options |= EPrintStaleReferencesOptions::Minimal;
+			}
 
-			UAssetManager::Get().GetAssetRegistry().EnumerateAssets(PluginArFilter, [&PluginName, &GameFeatureDataFilter, CheckForLoadedAsset](const FAssetData& AssetData)
+			// To minimize size of log, try to search for just public objects from the leaked packages.
+			TArray<UObject*> LeakedObjects;
+			for (UPackage* Package : LeakedPackages)
+			{
+				int32 StartCount = LeakedObjects.Num();
+				ForEachObjectWithPackage(Package, [&LeakedObjects](UObject* Object)
 				{
-					if (UAssetManager::Get().GetAssetRegistry().IsAssetIncludedByFilter(AssetData, GameFeatureDataFilter))
+					if (Object->HasAnyFlags(RF_Public))
 					{
-						return true;
+						LeakedObjects.Add(Object);
 					}
-
-					CheckForLoadedAsset(PluginName, AssetData);
-
 					return true;
-				});
-		}
-		else
-		{
-			UAssetManager::Get().GetAssetRegistry().EnumerateAssets(PluginArFilter, [&PluginName, CheckForLoadedAsset](const FAssetData& AssetData)
+				}, false);
+				if (LeakedObjects.Num() == StartCount)
 				{
-					CheckForLoadedAsset(PluginName, AssetData);
+					LeakedObjects.Add(Package);
+				}
+			}
 
-					return true;
-				});
+			UE_LOG(LogGameFeatures, Display, TEXT("Searching for references to %d leaked objects from plugin %s"), LeakedObjects.Num(), *PluginName);
+			GEngine->FindAndPrintStaleReferencesToObjects(LeakedObjects, Options);
 		}
 
-#endif // UE_BUILD_SHIPPING
+		// Rename the packages that we are streaming out so that we can possibly reload another copy of them
+		for (UPackage* Package : LeakedPackages)
+		{
+			UE_LOG(LogGameFeatures, Warning, TEXT("Marking leaking package %s as Garbage"), *Package->GetName());
+			ForEachObjectWithPackage(Package, [](UObject* Object)
+			{
+				Object->MarkAsGarbage();
+				return true;
+			}, false);
+			
+			Package->MarkAsGarbage();
+
+			if (!GIsEditor && !UObject::IsPendingKillEnabled())
+			{
+				UE::Object::RenameLeakedPackage(Package);
+			}
+		}
 	}
 
 #define GAME_FEATURE_PLUGIN_PROTOCOL_PREFIX(inEnum, inString) case EGameFeaturePluginProtocol::inEnum: return inString;
@@ -1788,7 +1866,7 @@ struct FGameFeaturePluginState_Unregistering : public FGameFeaturePluginState
 	{
 		if (bRequestedGC)
 		{
-			UE::GameFeatures::VerifyAssetsUnloaded(StateProperties.PluginName, false);
+			UE::GameFeatures::HandlePossibleAssetLeaks(StateProperties.PluginName);
 
 			StateStatus.SetTransition(EGameFeaturePluginState::Unmounting);
 			return;
@@ -1926,7 +2004,7 @@ struct FGameFeaturePluginState_Unloading : public FGameFeaturePluginState
 		if (bRequestedGC)
 		{
 #if !WITH_EDITOR // Disabled in editor since it's likely to report unloaded assets because of standalone packages
-			UE::GameFeatures::VerifyAssetsUnloaded(StateProperties.PluginName, true);
+			UE::GameFeatures::HandlePossibleAssetLeaks(StateProperties.PluginName, StateProperties.GameFeatureData ? StateProperties.GameFeatureData->GetOutermost() : nullptr);
 #endif //if !WITH_EDITOR
 
 			StateStatus.SetTransition(EGameFeaturePluginState::Registered);
diff --git a/Engine/Source/Runtime/CoreUObject/Private/UObject/ObjectRename.cpp b/Engine/Source/Runtime/CoreUObject/Private/UObject/ObjectRename.cpp
new file mode 100644
--- /dev/null
+++ b/Engine/Source/Runtime/CoreUObject/Private/UObject/ObjectRename.cpp
@@ -0,0 +1,20 @@
+// Copyright Epic Games, Inc. All Rights Reserved.
+
+#include "UObject/ObjectRename.h"
+
+#include "Misc/StringBuilder.h"
+#include "UObject/Object.h"
+#include "UObject/UObjectGlobalsInternal.h"
+
+namespace UE::Object
+{
+	void RenameLeakedPackage(UPackage* Package)
+	{
+		FCoreUObjectInternalDelegates::GetOnLeakedPackageRenameDelegate().Broadcast(Package);
+		FName NewName = MakeUniqueObjectName(nullptr, UPackage::StaticClass(), Package->GetFName());
+		TStringBuilder<FName::StringBufferSize> NewNameString;
+		NewNameString << NewName;
+		UE_LOG(LogObj, Log, TEXT("Renaming leaked package %s to %s"), *Package->GetName(), *NewNameString);
+		Package->Rename(*NewNameString, nullptr, REN_ForceNoResetLoaders | REN_DontCreateRedirectors | REN_NonTransactional);
+	}
+}
diff --git a/Engine/Source/Runtime/CoreUObject/Public/UObject/ObjectRename.h b/Engine/Source/Runtime/CoreUObject/Public/UObject/ObjectRename.h
new file mode 100644
--- /dev/null
+++ b/Engine/Source/Runtime/CoreUObject/Public/UObject/ObjectRename.h
@@ -0,0 +1,10 @@
+// Copyright Epic Games, Inc. All Rights Reserved.
+
+namespace UE::Object
+{
+	// Helper function to rename a package which ought to have been freed but which can't be GC'd
+	// because of outstanding references.
+	// Allows a new copy of that package to be loaded. 
+	// Flushes any internal caches (e.g. async loading) that would be broken by simply renaming the object.
+	COREUOBJECT_API void RenameLeakedPackage(UPackage* Package);
+}
