Fix use of unsubmitted new GC API.

diff --git a/Engine/Plugins/Experimental/GameFeatures/Source/GameFeatures/Private/GameFeaturePluginStateMachine.cpp b/Engine/Plugins/Experimental/GameFeatures/Source/GameFeatures/Private/GameFeaturePluginStateMachine.cpp
--- a/Engine/Plugins/Experimental/GameFeatures/Source/GameFeatures/Private/GameFeaturePluginStateMachine.cpp
+++ b/Engine/Plugins/Experimental/GameFeatures/Source/GameFeatures/Private/GameFeaturePluginStateMachine.cpp
@@ -116,7 +116,6 @@ namespace UE::GameFeatures
 		{
 			// If the UObject hash knew about package mount roots, we could avoid this loop
 			TRACE_CPUPROFILER_EVENT_SCOPE(PackageLoop);
-			FPermanentObjectPoolExtents PermExtents;
 			ForEachObjectOfClass(UPackage::StaticClass(), [&](UObject* Obj)
 			{
 				if (UPackage* Package = CastChecked<UPackage>(Obj, ECastCheckedType::NullAllowed))
@@ -125,7 +124,7 @@ namespace UE::GameFeatures
 					{
 						return;
 					}
-					if (PermExtents.Contains(Package))
+					if (GUObjectAllocator.ResidesInPermanentPool(Package))
 					{
 						return;
 					}
