[Backout] - CL23315156
Original CL Desc
-----------------------------------------------------------------
Additional error reporting for object leaks when unmounting GFPs.
Rename objects that leak so that new copies of them can be reloaded if the GFP plugin is mounted again.

diff --git a/Engine/Plugins/Experimental/GameFeatures/Source/GameFeatures/Private/GameFeaturePluginStateMachine.cpp b/Engine/Plugins/Experimental/GameFeatures/Source/GameFeatures/Private/GameFeaturePluginStateMachine.cpp
--- a/Engine/Plugins/Experimental/GameFeatures/Source/GameFeatures/Private/GameFeaturePluginStateMachine.cpp
+++ b/Engine/Plugins/Experimental/GameFeatures/Source/GameFeatures/Private/GameFeaturePluginStateMachine.cpp
@@ -25,8 +25,7 @@
 #include "GameFeaturesSubsystem.h"
 #include "GameFeaturesProjectPolicies.h"
 #include "Containers/Ticker.h"
-#include "UObject/ObjectRename.h"
-#include "UObject/UObjectAllocator.h"
+#include "UObject/ReferenceChainSearch.h"
 
 #include UE_INLINE_GENERATED_CPP_BY_NAME(GameFeaturePluginStateMachine)
 
@@ -45,44 +44,15 @@ namespace UE::GameFeatures
 		ShouldLogMountedFiles,
 		TEXT("Should the newly mounted files be logged."));
 
-	static bool GRenameLeakedPackages = true;
-	static FAutoConsoleVariableRef CVarRenameLeakedPackages(
-		TEXT("GameFeaturePlugin.LeakedAssetTrace.RenameLeakedPackages"),
-		GRenameLeakedPackages,
-		TEXT("Should packages which are leaked after the Game Feature Plugin is unloaded or unmounted."),
-		ECVF_Default
-	);
-
-	static int32 GLeakedAssetTrace_Severity = 2;
-	static FAutoConsoleVariableRef CVarLeakedAssetTrace_Severity(
-	#if UE_BUILD_SHIPPING
-		TEXT("GameFeaturePlugin.LeakedAssetTrace.Severity.Shipping"),
-	#else
-		TEXT("GameFeaturePlugin.LeakedAssetTrace.Severity"),
-	#endif
-		GLeakedAssetTrace_Severity,
-		TEXT("Controls severity of logging when the engine detects that assets from an Game Feature Plugin were leaked during unloading or unmounting.\n")
-		TEXT("0 - all reference tracing and logging is disabled\n")
-		TEXT("1 - logs an error\n")
-		TEXT("2 - ensure\n")
-		TEXT("3 - fatal error\n"),
-		ECVF_Default
-	);
-
-	static int32 GLeakedAssetTrace_TraceMode = (UE_BUILD_SHIPPING ? 0 : 1);
-	static FAutoConsoleVariableRef CVarLeakedAssetTrace_TraceMode(
-	#if UE_BUILD_SHIPPING
-		TEXT("GameFeaturePlugin.LeakedAssetTrace.TraceMode.Shipping"),
-	#else
-		TEXT("GameFeaturePlugin.LeakedASsetTrace.TraceMode"),
-	#endif
-		GLeakedAssetTrace_TraceMode,
-		TEXT("Controls detail level of reference tracing when the engine detects that assets from a Game Feature Plugin were leaked during unloading or unmounting.\n")
-		TEXT("0 - direct references only\n")
-		TEXT("1 - full reference trace"),
-		ECVF_Default
-	);
+	static TAutoConsoleVariable<bool> CVarVerifyPluginUnload(TEXT("GameFeaturePlugin.VerifyUnload"), 
+		true,
+		TEXT("Verify plugin assets are no longer in memory when unloading."),
+		ECVF_Default);
 
+	static TAutoConsoleVariable<bool> CVarVerifyPluginUnloadDumpChains(TEXT("GameFeaturePlugin.VerifyUnloadDumpChains"),
+		false,
+		TEXT("Dump reference chains for any detected plugin asset leaks."),
+		ECVF_Default);
 	#define GAME_FEATURE_PLUGIN_STATE_TO_STRING(inEnum, inText) case EGameFeaturePluginState::inEnum: return TEXT(#inEnum);
 	FString ToString(EGameFeaturePluginState InType)
 	{
@@ -96,116 +66,68 @@ namespace UE::GameFeatures
 	}
 	#undef GAME_FEATURE_PLUGIN_STATE_TO_STRING
 
-	// Check if any assets from the plugin mount point have leaked, and if so trace them.
-	// Then rename them to allow new copies of them to be loaded. 
-	void HandlePossibleAssetLeaks(const FString& PluginName, UPackage* IgnorePackage = nullptr)
+	// Verify that all assets from this plugin have been unloaded and GC'd	
+	void VerifyAssetsUnloaded(const FString& PluginName, bool bIgnoreGameFeatureData)
 	{
-		TRACE_CPUPROFILER_EVENT_SCOPE(HandlePossibleAssetLeaks);
-
-		const bool bFindLeakedPackages = GLeakedAssetTrace_Severity != 0 || GRenameLeakedPackages;
-		if (!bFindLeakedPackages)
+#if (!UE_BUILD_SHIPPING || UE_SERVER || WITH_EDITOR)
+		if (!UE::GameFeatures::CVarVerifyPluginUnload.GetValueOnGameThread())
 		{
 			return;
 		}
-		
-		TSet<UPackage*> LeakedPackages;
-		TStringBuilder<512> Prefix;
-		Prefix << '/' << PluginName << '/';
 
-		TStringBuilder<NAME_SIZE> NameBuffer;
+		FARFilter PluginArFilter;
+		PluginArFilter.PackagePaths.Add(FName(TEXT("/") + PluginName));
+		PluginArFilter.bRecursivePaths = true;
+
+		auto CheckForLoadedAsset = [](const FString& PluginName, const FAssetData &AssetData)
 		{
-			// If the UObject hash knew about package mount roots, we could avoid this loop
-			TRACE_CPUPROFILER_EVENT_SCOPE(PackageLoop);
-			FPermanentObjectPoolExtents PermExtents;
-			ForEachObjectOfClass(UPackage::StaticClass(), [&](UObject* Obj)
+			if (AssetData.IsAssetLoaded())
 			{
-				if (UPackage* Package = CastChecked<UPackage>(Obj, ECastCheckedType::NullAllowed))
+				UE_LOG(LogGameFeatures, Error, TEXT("GFP %s failed to unload asset %s!"), *PluginName, *AssetData.GetFullName());
+
+				if (CVarVerifyPluginUnloadDumpChains.GetValueOnGameThread())
 				{
-					if (Package == IgnorePackage)
-					{
-						return;
-					}
-					if (PermExtents.Contains(Package))
-					{
-						return;
-					}
-					NameBuffer.Reset();
-					Package->GetFName().GetDisplayNameEntry()->AppendNameToString(NameBuffer);
-					if (NameBuffer.ToView().StartsWith(Prefix, ESearchCase::IgnoreCase))
-					{
-						LeakedPackages.Add(Package);
-					}
+					UObject* AssetObj = AssetData.GetAsset();
+					//EInternalObjectFlags InternalFlags = AssetObj->GetInternalFlags();
+					FReferenceChainSearch(AssetObj, EReferenceChainSearchMode::Shortest | EReferenceChainSearchMode::PrintResults);
 				}
-			});
-		}
 
-		if (LeakedPackages.Num() == 0)
-		{
-			return;
-		}
+				ensureAlwaysMsgf(false, TEXT("GFP %s failed to unload asset %s!"), *PluginName, *AssetData.GetFullName());
+			}
+		};
 
-		if (GLeakedAssetTrace_Severity != 0)
+		if (bIgnoreGameFeatureData)
 		{
-			EPrintStaleReferencesOptions Options = EPrintStaleReferencesOptions::None;
-			switch (GLeakedAssetTrace_Severity)
-			{
-			case 3:
-				Options = EPrintStaleReferencesOptions::Fatal;
-				break;
-			case 2: 
-				Options = EPrintStaleReferencesOptions::Ensure | EPrintStaleReferencesOptions::Error;
-				break;
-			case 1:
-			default:
-				Options = EPrintStaleReferencesOptions::Error;
-				break;
-			}	
+			FARFilter RawGameFeatureDataFilter;
+			RawGameFeatureDataFilter.ClassPaths.Add(UGameFeatureData::StaticClass()->GetClassPathName());
+			RawGameFeatureDataFilter.bRecursiveClasses = true;
 
-			if (GLeakedAssetTrace_TraceMode == 0)
-			{
-				Options |= EPrintStaleReferencesOptions::Minimal;
-			}
+			FARCompiledFilter GameFeatureDataFilter;
+			UAssetManager::Get().GetAssetRegistry().CompileFilter(RawGameFeatureDataFilter, GameFeatureDataFilter);
 
-			// To minimize size of log, try to search for just public objects from the leaked packages.
-			TArray<UObject*> LeakedObjects;
-			for (UPackage* Package : LeakedPackages)
-			{
-				int32 StartCount = LeakedObjects.Num();
-				ForEachObjectWithPackage(Package, [&LeakedObjects](UObject* Object)
+			UAssetManager::Get().GetAssetRegistry().EnumerateAssets(PluginArFilter, [&PluginName, &GameFeatureDataFilter, CheckForLoadedAsset](const FAssetData& AssetData)
 				{
-					if (Object->HasAnyFlags(RF_Public))
+					if (UAssetManager::Get().GetAssetRegistry().IsAssetIncludedByFilter(AssetData, GameFeatureDataFilter))
 					{
-						LeakedObjects.Add(Object);
+						return true;
 					}
-					return true;
-				}, false);
-				if (LeakedObjects.Num() == StartCount)
-				{
-					LeakedObjects.Add(Package);
-				}
-			}
 
-			UE_LOG(LogGameFeatures, Display, TEXT("Searching for references to %d leaked objects from plugin %s"), LeakedObjects.Num(), *PluginName);
-			GEngine->FindAndPrintStaleReferencesToObjects(LeakedObjects, Options);
-		}
+					CheckForLoadedAsset(PluginName, AssetData);
 
-		// Rename the packages that we are streaming out so that we can possibly reload another copy of them
-		for (UPackage* Package : LeakedPackages)
+					return true;
+				});
+		}
+		else
 		{
-			UE_LOG(LogGameFeatures, Warning, TEXT("Marking leaking package %s as Garbage"), *Package->GetName());
-			ForEachObjectWithPackage(Package, [](UObject* Object)
-			{
-				Object->MarkAsGarbage();
-				return true;
-			}, false);
-			
-			Package->MarkAsGarbage();
+			UAssetManager::Get().GetAssetRegistry().EnumerateAssets(PluginArFilter, [&PluginName, CheckForLoadedAsset](const FAssetData& AssetData)
+				{
+					CheckForLoadedAsset(PluginName, AssetData);
 
-			if (!GIsEditor && !UObject::IsPendingKillEnabled())
-			{
-				UE::Object::RenameLeakedPackage(Package);
-			}
+					return true;
+				});
 		}
+
+#endif // UE_BUILD_SHIPPING
 	}
 
 #define GAME_FEATURE_PLUGIN_PROTOCOL_PREFIX(inEnum, inString) case EGameFeaturePluginProtocol::inEnum: return inString;
@@ -1866,7 +1788,7 @@ struct FGameFeaturePluginState_Unregistering : public FGameFeaturePluginState
 	{
 		if (bRequestedGC)
 		{
-			UE::GameFeatures::HandlePossibleAssetLeaks(StateProperties.PluginName);
+			UE::GameFeatures::VerifyAssetsUnloaded(StateProperties.PluginName, false);
 
 			StateStatus.SetTransition(EGameFeaturePluginState::Unmounting);
 			return;
@@ -2004,7 +1926,7 @@ struct FGameFeaturePluginState_Unloading : public FGameFeaturePluginState
 		if (bRequestedGC)
 		{
 #if !WITH_EDITOR // Disabled in editor since it's likely to report unloaded assets because of standalone packages
-			UE::GameFeatures::HandlePossibleAssetLeaks(StateProperties.PluginName, StateProperties.GameFeatureData ? StateProperties.GameFeatureData->GetOutermost() : nullptr);
+			UE::GameFeatures::VerifyAssetsUnloaded(StateProperties.PluginName, true);
 #endif //if !WITH_EDITOR
 
 			StateStatus.SetTransition(EGameFeaturePluginState::Registered);
diff --git a/Engine/Source/Runtime/CoreUObject/Private/UObject/ObjectRename.cpp b/Engine/Source/Runtime/CoreUObject/Private/UObject/ObjectRename.cpp
deleted file mode 100644
diff --git a/Engine/Source/Runtime/CoreUObject/Public/UObject/ObjectRename.h b/Engine/Source/Runtime/CoreUObject/Public/UObject/ObjectRename.h
deleted file mode 100644
