#UE Avoid IO time by re-reading uplugin files. Was potentially causing a hang on platforms with slow or inhitibted I/O speeds. Can be disabled by setting cvar GameFeaturePlugin.CachePluginDetails to false.

#localization none
#tests Editor and -game tests locally. Logged in, loaded plugins, verified cached values were read after a hotfix check

diff --git a/Engine/Plugins/Experimental/GameFeatures/Source/GameFeatures/Private/GameFeaturesSubsystem.cpp b/Engine/Plugins/Experimental/GameFeatures/Source/GameFeatures/Private/GameFeaturesSubsystem.cpp
--- a/Engine/Plugins/Experimental/GameFeatures/Source/GameFeatures/Private/GameFeaturesSubsystem.cpp
+++ b/Engine/Plugins/Experimental/GameFeatures/Source/GameFeatures/Private/GameFeaturesSubsystem.cpp
@@ -8,6 +8,7 @@
 #include "GameFeaturePluginStateMachine.h"
 #include "GameFeatureStateChangeObserver.h"
 #include "GameplayTagsManager.h"
+#include "HAL/FileManager.h"
 #include "Interfaces/IPluginManager.h"
 #include "Misc/Paths.h"
 #include "Misc/FileHelper.h"
@@ -54,6 +55,13 @@ namespace UE::GameFeatures
 
 		const FText GenericError = NSLOCTEXT("GameFeatures", "CommonErrors.Generic", "An error has occurred. Please try again later.");
 	}
+
+	static bool GCachePluginDetails = true;
+	static FAutoConsoleVariableRef CVarRenameLeakedPackages(
+		TEXT("GameFeaturePlugin.CachePluginDetails"),
+		GCachePluginDetails,
+		TEXT("Should use plugin details caching."),
+		ECVF_Default);
 }
 
 #define GAME_FEATURE_PLUGIN_STATE_LEX_TO_STRING(inEnum, inText)  \
@@ -1253,10 +1261,28 @@ EGameFeaturePluginState UGameFeaturesSubsystem::GetPluginState(FGameFeaturePlugi
 bool UGameFeaturesSubsystem::GetGameFeaturePluginDetails(const FString& PluginDescriptorFilename, FGameFeaturePluginDetails& OutPluginDetails) const
 {
 	TRACE_CPUPROFILER_EVENT_SCOPE(GFP_GetPluginDetails);
-	// @TODO: We load the descriptor 2-3 per plugin because FPluginReferenceDescriptor doesn't cache any of this info.
+
 	// GFPs are implemented with a plugin so FPluginReferenceDescriptor doesn't know anything about them.
 	// Need a better way of storing GFP specific plugin data...
 
+	FDateTime FileTimeStamp;
+	if (UE::GameFeatures::GCachePluginDetails)
+	{
+		FileTimeStamp = IFileManager::Get().GetTimeStamp(*PluginDescriptorFilename);
+		if (FCachedGameFeaturePluginDetails* ExistingDetails = CachedPluginDetailsByFilename.Find(PluginDescriptorFilename))
+		{
+			if (ExistingDetails->TimeStamp == FileTimeStamp)
+			{
+				OutPluginDetails = ExistingDetails->Details;
+				return true;
+			}
+			else
+			{
+				CachedPluginDetailsByFilename.Remove(PluginDescriptorFilename);
+			}
+		}
+	}
+
 	// Read the file to a string
 	FString FileContents;
 	if (!FFileHelper::LoadFileToString(FileContents, *PluginDescriptorFilename))
@@ -1349,6 +1375,10 @@ bool UGameFeaturesSubsystem::GetGameFeaturePluginDetails(const FString& PluginDe
 		}
 	}
 
+	if (UE::GameFeatures::GCachePluginDetails)
+	{
+		CachedPluginDetailsByFilename.Add(PluginDescriptorFilename, FCachedGameFeaturePluginDetails(OutPluginDetails, FileTimeStamp));
+	}
 	return true;
 }
 
diff --git a/Engine/Plugins/Experimental/GameFeatures/Source/GameFeatures/Public/GameFeaturesSubsystem.h b/Engine/Plugins/Experimental/GameFeatures/Source/GameFeatures/Public/GameFeaturesSubsystem.h
--- a/Engine/Plugins/Experimental/GameFeatures/Source/GameFeatures/Public/GameFeaturesSubsystem.h
+++ b/Engine/Plugins/Experimental/GameFeatures/Source/GameFeatures/Public/GameFeaturesSubsystem.h
@@ -607,6 +607,15 @@ private:
 
 	TMap<FString, FString> GameFeaturePluginNameToPathMap;
 
+	struct FCachedGameFeaturePluginDetails
+	{
+		FGameFeaturePluginDetails Details;
+		FDateTime TimeStamp;
+		FCachedGameFeaturePluginDetails() {}
+		FCachedGameFeaturePluginDetails(const FGameFeaturePluginDetails& InDetails, const FDateTime& InTimeStamp) : Details(InDetails), TimeStamp(InTimeStamp) {}
+	};
+	mutable TMap<FString, FCachedGameFeaturePluginDetails> CachedPluginDetailsByFilename;
+
 	UPROPERTY()
 	TArray<TObjectPtr<UObject>> Observers;
 
