#UE Fix cvar name collision

diff --git a/Engine/Plugins/Experimental/GameFeatures/Source/GameFeatures/Private/GameFeaturesSubsystem.cpp b/Engine/Plugins/Experimental/GameFeatures/Source/GameFeatures/Private/GameFeaturesSubsystem.cpp
--- a/Engine/Plugins/Experimental/GameFeatures/Source/GameFeatures/Private/GameFeaturesSubsystem.cpp
+++ b/Engine/Plugins/Experimental/GameFeatures/Source/GameFeatures/Private/GameFeaturesSubsystem.cpp
@@ -57,7 +57,7 @@ namespace UE::GameFeatures
 	}
 
 	static bool GCachePluginDetails = true;
-	static FAutoConsoleVariableRef CVarRenameLeakedPackages(
+	static FAutoConsoleVariableRef CVarCachePluginDetails(
 		TEXT("GameFeaturePlugin.CachePluginDetails"),
 		GCachePluginDetails,
 		TEXT("Should use plugin details caching."),
