Restore GameFeaturePlugin.VerifyUnload and add GameFeaturePlugin.VerifyUnload.SkipList to allow specific GFPs to opt-out from leak detection.

diff --git a/Engine/Plugins/Experimental/GameFeatures/Source/GameFeatures/Private/GameFeaturePluginStateMachine.cpp b/Engine/Plugins/Experimental/GameFeatures/Source/GameFeatures/Private/GameFeaturePluginStateMachine.cpp
--- a/Engine/Plugins/Experimental/GameFeatures/Source/GameFeatures/Private/GameFeaturePluginStateMachine.cpp
+++ b/Engine/Plugins/Experimental/GameFeatures/Source/GameFeatures/Private/GameFeaturePluginStateMachine.cpp
@@ -12,15 +12,18 @@
 #include "InstallBundleUtils.h"
 #include "BundlePrereqCombinedStatusHelper.h"
 #include "Interfaces/IPluginManager.h"
+#include "Misc/AsciiSet.h"
 #include "Misc/ConfigCacheIni.h"
 #include "Misc/CoreDelegates.h"
 #include "Misc/FileHelper.h"
 #include "Misc/Paths.h"
+#include "Misc/WildcardString.h"
 #include "Algo/AllOf.h"
 #include "ProfilingDebugging/CpuProfilerTrace.h"
 #include "Serialization/ArrayReader.h"
 #include "Serialization/MemoryReader.h"
 #include "Stats/Stats.h"
+#include "String/ParseTokens.h"
 #include "EngineUtils.h"
 #include "GameFeaturesSubsystem.h"
 #include "GameFeaturesProjectPolicies.h"
@@ -45,6 +48,19 @@ namespace UE::GameFeatures
 		ShouldLogMountedFiles,
 		TEXT("Should the newly mounted files be logged."));
 
+	static bool GVerifyUnload = true;
+	static FAutoConsoleVariableRef CVarVerifyPluginUnload(TEXT("GameFeaturePlugin.VerifyUnload"),
+		GVerifyUnload,
+		TEXT("Verify plugin assets are no longer in memory when unloading."),
+		ECVF_Default);
+
+	static FString GVerifyPluginSkipList;
+	static FAutoConsoleVariableRef CVarVerifyPluginSkipList(TEXT("GameFeaturePlugin.VerifyUnload.SkipList"),
+		GVerifyPluginSkipList,
+		TEXT("Comma-separated list of names of plugins for which to skip verification."),
+		ECVF_Default
+		);
+
 	static bool GRenameLeakedPackages = true;
 	static FAutoConsoleVariableRef CVarRenameLeakedPackages(
 		TEXT("GameFeaturePlugin.LeakedAssetTrace.RenameLeakedPackages"),
@@ -102,11 +118,34 @@ namespace UE::GameFeatures
 	{
 		TRACE_CPUPROFILER_EVENT_SCOPE(HandlePossibleAssetLeaks);
 
-		const bool bFindLeakedPackages = GLeakedAssetTrace_Severity != 0 || GRenameLeakedPackages;
+		const bool bFindLeakedPackages = GVerifyUnload && (GLeakedAssetTrace_Severity != 0 || GRenameLeakedPackages);
 		if (!bFindLeakedPackages)
 		{
 			return;
 		}
+
+		static const FAsciiSet Wildcards("*?");
+		bool bSkip = false;
+		UE::String::ParseTokens(MakeStringView(GVerifyPluginSkipList), TEXTVIEW(","), [&PluginName, &bSkip](FStringView Item) {
+			if (bSkip) { return; }
+			if (Item.Equals(PluginName, ESearchCase::IgnoreCase))
+			{
+				bSkip = true;
+			}
+			else if (FAsciiSet::HasAny(Item, Wildcards))
+			{
+				FString Pattern = FString(Item); // Need to copy to null terminate
+				if (FWildcardString::IsMatchSubstring(*Pattern, *PluginName, *PluginName + PluginName.Len(), ESearchCase::IgnoreCase))
+				{
+					bSkip = true;
+				}
+			}
+		});
+
+		if (bSkip)
+		{
+			return;
+		}
 		
 		TSet<UPackage*> LeakedPackages;
 		TStringBuilder<512> Prefix;
