#UE Report total BuiltIn GameFeaturePlugin load time.

Also, in non-shipping builds, report builtin GameFeaturePlugins that take a long time to load in the log.
Thresholds for this reporting are configurable via cvar as follows:

GameFeaturePlugin.BuiltInPluginLoadTimeReportThreshold
Built-in plugins that take longer than this amount of time, in seconds, will be reported to the log during startup in non-shipping builds.

GameFeaturePlugin.BuiltInPluginLoadTimeErrorThreshold
Built-in plugins that take longer than this amount of time, in seconds, will cause an error during startup in non-shipping builds.

GameFeaturePlugin.BuiltInPluginLoadTimeMaxReportCount
When listing worst offenders for built-in plugin load time, show no more than this many plugins, to reduce log spam.

diff --git a/Engine/Plugins/Experimental/GameFeatures/Source/GameFeatures/Private/GameFeaturesSubsystem.cpp b/Engine/Plugins/Experimental/GameFeatures/Source/GameFeatures/Private/GameFeaturesSubsystem.cpp
--- a/Engine/Plugins/Experimental/GameFeatures/Source/GameFeatures/Private/GameFeaturesSubsystem.cpp
+++ b/Engine/Plugins/Experimental/GameFeatures/Source/GameFeatures/Private/GameFeaturesSubsystem.cpp
@@ -59,6 +59,29 @@ namespace UE::GameFeatures
 		GCachePluginDetails,
 		TEXT("Should use plugin details caching."),
 		ECVF_Default);
+
+#if !UE_BUILD_SHIPPING
+	static float GBuiltInPluginLoadTimeReportThreshold = 3.0;
+	static FAutoConsoleVariableRef CVarBuiltInPluginLoadTimeReportThreshold(
+		TEXT("GameFeaturePlugin.BuiltInPluginLoadTimeReportThreshold"),
+		GBuiltInPluginLoadTimeReportThreshold,
+		TEXT("Built-in plugins that take longer than this amount of time, in seconds, will be reported to the log during startup in non-shipping builds."),
+		ECVF_Default);
+
+	static float GBuiltInPluginLoadTimeErrorThreshold = 600;
+	static FAutoConsoleVariableRef CVarBuiltInPluginLoadTimeErrorThreshold(
+		TEXT("GameFeaturePlugin.BuiltInPluginLoadTimeErrorThreshold"),
+		GBuiltInPluginLoadTimeErrorThreshold,
+		TEXT("Built-in plugins that take longer than this amount of time, in seconds, will cause an error during startup in non-shipping builds."),
+		ECVF_Default);
+
+	static int32 GBuiltInPluginLoadTimeMaxReportCount = 10;
+	static FAutoConsoleVariableRef CVarBuiltInPluginLoadTimeMaxReportCount(
+		TEXT("GameFeaturePlugin.BuiltInPluginLoadTimeMaxReportCount"),
+		GBuiltInPluginLoadTimeMaxReportCount,
+		TEXT("When listing worst offenders for built-in plugin load time, show no more than this many plugins, to reduce log spam."),
+		ECVF_Default);
+#endif // !UE_BUILD_SHIPPING
 }
 
 #define GAME_FEATURE_PLUGIN_STATE_LEX_TO_STRING(inEnum, inText)  \
@@ -1163,14 +1186,86 @@ void UGameFeaturesSubsystem::LoadBuiltInGameFeaturePlugin(const TSharedRef<IPlug
 	UAssetManager::Get().PopBulkScanning();
 }
 
+#if UE_BUILD_SHIPPING
+class FBuiltInPluginLoadTimeTracker {};
+class FBuiltInPluginLoadTimeTrackerScope
+{
+public:
+	FORCEINLINE FBuiltInPluginLoadTimeTrackerScope(FBuiltInPluginLoadTimeTracker& InTracker, const TSharedRef<IPlugin>& Plugin) {};
+};
+#else // !UE_BUILD_SHIPPING
+class FBuiltInPluginLoadTimeTracker
+{
+public:
+	~FBuiltInPluginLoadTimeTracker()
+	{
+		if (PluginLoadTimes.Num() > 0)
+		{
+			UE_LOG(LogGameFeatures, Display, TEXT("There were %d built in plugins that took longer than %.4fs to load. Listing worst offenders."), PluginLoadTimes.Num(), UE::GameFeatures::GBuiltInPluginLoadTimeReportThreshold);
+			const int32 NumToReport = FMath::Min(PluginLoadTimes.Num(), UE::GameFeatures::GBuiltInPluginLoadTimeMaxReportCount);
+			Algo::Sort(PluginLoadTimes, [](const TPair<FString, double>& A, TPair<FString, double>& B) { return A.Value > B.Value; });
+			for (int32 PluginIdx = 0; PluginIdx < NumToReport; ++PluginIdx)
+			{
+				const TPair<FString, double>& Plugin = PluginLoadTimes[PluginIdx];
+				double LoadTime = Plugin.Value;
+				if (LoadTime >= UE::GameFeatures::GBuiltInPluginLoadTimeErrorThreshold)
+				{
+					UE_LOG(LogGameFeatures, Error, TEXT("%s took %.4f seconds to load. Something was done to significantly increase the load time of this plugin and it is now well outside what is acceptable. Reduce the load time to much less than %.4f seconds. Ideally, reduce the load time to less than %.4f seconds."),
+						*Plugin.Key, Plugin.Value, UE::GameFeatures::GBuiltInPluginLoadTimeErrorThreshold, UE::GameFeatures::GBuiltInPluginLoadTimeReportThreshold);
+				}
+				else
+				{
+					UE_LOG(LogGameFeatures, Display, TEXT("  %.4fs\t%s"), Plugin.Value, *Plugin.Key);
+				}
+			}
+		}
+	}
+
+	void ReportPlugin(const FString& PluginName, double LoadTime)
+	{
+		PluginLoadTimes.Emplace(PluginName, LoadTime);
+	}
+
+private:
+	TArray<TPair<FString, double>> PluginLoadTimes;
+	
+};
+
+class FBuiltInPluginLoadTimeTrackerScope
+{
+public:
+	FBuiltInPluginLoadTimeTrackerScope(FBuiltInPluginLoadTimeTracker& InTracker, const TSharedRef<IPlugin>& Plugin)
+		: Tracker(InTracker), PluginName(Plugin->GetName()), StartTime(FPlatformTime::Seconds())
+	{}
+
+	~FBuiltInPluginLoadTimeTrackerScope()
+	{
+		double LoadTime = FPlatformTime::Seconds() - StartTime;
+		if (LoadTime >= UE::GameFeatures::GBuiltInPluginLoadTimeReportThreshold)
+		{
+			Tracker.ReportPlugin(PluginName, LoadTime);
+		}
+	}
+
+private:
+	FBuiltInPluginLoadTimeTracker& Tracker;
+	FString PluginName;
+	double StartTime;
+};
+
+#endif // !UE_BUILD_SHIPPING
+
 void UGameFeaturesSubsystem::LoadBuiltInGameFeaturePlugins(FBuiltInPluginAdditionalFilters AdditionalFilter)
 {
+	SCOPE_LOG_TIME_IN_SECONDS(TEXT("BuiltInGameFeaturePlugins loaded."), nullptr);
 	UAssetManager::Get().PushBulkScanning();
 	UGameplayTagsManager::Get().PushDeferOnGameplayTagTreeChangedBroadcast();
 
+	FBuiltInPluginLoadTimeTracker PluginLoadTimeTracker;
 	TArray<TSharedRef<IPlugin>> EnabledPlugins = IPluginManager::Get().GetEnabledPlugins();
 	for (const TSharedRef<IPlugin>& Plugin : EnabledPlugins)
 	{
+		FBuiltInPluginLoadTimeTrackerScope TrackerScope(PluginLoadTimeTracker, Plugin);
 		LoadBuiltInGameFeaturePlugin(Plugin, AdditionalFilter);
 	}
 
