#UE IWYU fix

diff --git a/Engine/Plugins/Experimental/GameFeatures/Source/GameFeatures/Private/GameFeaturesSubsystem.cpp b/Engine/Plugins/Experimental/GameFeatures/Source/GameFeatures/Private/GameFeaturesSubsystem.cpp
--- a/Engine/Plugins/Experimental/GameFeatures/Source/GameFeatures/Private/GameFeaturesSubsystem.cpp
+++ b/Engine/Plugins/Experimental/GameFeatures/Source/GameFeatures/Private/GameFeaturesSubsystem.cpp
@@ -14,6 +14,7 @@
 #include "Misc/Paths.h"
 #include "Misc/FileHelper.h"
 #include "Serialization/JsonSerializer.h"
+#include "Stats/StatsMisc.h"
 #include "Engine/AssetManager.h"
 #include "Engine/StreamableManager.h"
 #include "Engine/AssetManagerSettings.h"
