Install bundle protocol now disables INI loading by default. It can be enabled with a flag

diff --git a/Engine/Plugins/Experimental/GameFeatures/Source/GameFeatures/Private/GameFeaturePluginStateMachine.cpp b/Engine/Plugins/Experimental/GameFeatures/Source/GameFeatures/Private/GameFeaturePluginStateMachine.cpp
--- a/Engine/Plugins/Experimental/GameFeatures/Source/GameFeatures/Private/GameFeaturePluginStateMachine.cpp
+++ b/Engine/Plugins/Experimental/GameFeatures/Source/GameFeatures/Private/GameFeaturePluginStateMachine.cpp
@@ -346,9 +346,10 @@ namespace UE::GameFeatures
 	XOPTION(Bundles,					TEXT("Bundles"))					\
 	XOPTION(Version,					TEXT("V"))							\
 	XOPTION(Flags,						TEXT("Flags"))						\
-	XOPTION(ReleaseFlags,				TEXT("ReleaseFlags"))						\
+	XOPTION(ReleaseFlags,				TEXT("ReleaseFlags"))				\
 	XOPTION(UninstallBeforeTerminate,	TEXT("UninstallBeforeTerminate"))	\
-	XOPTION(UserPauseDownload,			TEXT("Paused"))
+	XOPTION(UserPauseDownload,			TEXT("Paused"))						\
+	XOPTION(AllowIniLoading,			TEXT("AllowIni"))
 
 #define INSTALL_BUNDLE_PROTOCOL_URL_OPTIONS_ENUM(inEnum, inString) inEnum,
 enum class EGameFeatureInstallBundleProtocolOptions : uint8
@@ -567,6 +568,25 @@ bool FGameFeaturePluginState::ShouldVisitUninstallStateBeforeTerminal() const
 	}
 }
 
+bool FGameFeaturePluginState::AllowIniLoading() const
+{
+	switch (StateProperties.GetPluginProtocol())
+	{
+	case (EGameFeaturePluginProtocol::InstallBundle):
+	{
+		//InstallBundleProtocol's have a MetaData that controlls if INI loading is allowed
+		//The protocol default is not to allow INI loading since the source is likely untrusted
+		return StateProperties.ProtocolMetadata.GetSubtype<FInstallBundlePluginProtocolMetaData>().bAllowIniLoading;
+	}
+
+	//Default behavior is to allow INI loading
+	default:
+	{
+		return true;
+	}
+	}
+}
+
 /*
 =========================================================
   States
@@ -1945,22 +1965,26 @@ struct FGameFeaturePluginState_Registering : public FGameFeaturePluginState
 		const FString BackupGameFeatureDataPath = FString::Printf(TEXT("/%s/%s.%s"), *StateProperties.PluginName, *StateProperties.PluginName, *StateProperties.PluginName);
 
 		FString PreferredGameFeatureDataPath = TEXT("/") + StateProperties.PluginName + TEXT("/GameFeatureData.GameFeatureData");
-		// Allow game feature location to be overriden globally and from within the plugin
-		FString OverrideIniPathName = StateProperties.PluginName + TEXT("_Override");
-		FString OverridePath = GConfig->GetStr(TEXT("GameFeatureData"), *OverrideIniPathName, GGameIni);
-		if (OverridePath.IsEmpty())
+
+		if (AllowIniLoading())
 		{
-			const FString SettingsOverride = PluginFolder / TEXT("Config") / TEXT("Settings.ini");
-			if (FPaths::FileExists(SettingsOverride))
+			// Allow game feature location to be overriden globally and from within the plugin
+			FString OverrideIniPathName = StateProperties.PluginName + TEXT("_Override");
+			FString OverridePath = GConfig->GetStr(TEXT("GameFeatureData"), *OverrideIniPathName, GGameIni);
+			if (OverridePath.IsEmpty())
 			{
-				GConfig->LoadFile(SettingsOverride);
-				OverridePath = GConfig->GetStr(TEXT("GameFeatureData"), TEXT("Override"), SettingsOverride);
-				GConfig->UnloadFile(SettingsOverride);
+				const FString SettingsOverride = PluginFolder / TEXT("Config") / TEXT("Settings.ini");
+				if (FPaths::FileExists(SettingsOverride))
+				{
+					GConfig->LoadFile(SettingsOverride);
+					OverridePath = GConfig->GetStr(TEXT("GameFeatureData"), TEXT("Override"), SettingsOverride);
+					GConfig->UnloadFile(SettingsOverride);
+				}
+			}
+			if (!OverridePath.IsEmpty())
+			{
+				PreferredGameFeatureDataPath = OverridePath;
 			}
-		}
-		if (!OverridePath.IsEmpty())
-		{
-			PreferredGameFeatureDataPath = OverridePath;
 		}
 		
 		auto LoadGameFeatureData = [](const FString& Path) -> UGameFeatureData*
@@ -1988,7 +2012,11 @@ struct FGameFeaturePluginState_Registering : public FGameFeaturePluginState
 
 		if (StateProperties.GameFeatureData)
 		{
-			StateProperties.GameFeatureData->InitializeBasePluginIniFile(StateProperties.PluginInstalledFilename);
+			if (AllowIniLoading())
+			{
+				StateProperties.GameFeatureData->InitializeBasePluginIniFile(StateProperties.PluginInstalledFilename);
+			}
+
 			StateStatus.SetTransition(EGameFeaturePluginState::Registered);
 
 			check(StateProperties.AddedPrimaryAssetTypes.Num() == 0);
@@ -2246,7 +2274,10 @@ struct FGameFeaturePluginState_Activating : public FGameFeaturePluginState
 
 		FGameFeatureActivatingContext Context;
 
-		StateProperties.GameFeatureData->InitializeHierarchicalPluginIniFiles(StateProperties.PluginInstalledFilename);
+		if (AllowIniLoading())
+		{
+			StateProperties.GameFeatureData->InitializeHierarchicalPluginIniFiles(StateProperties.PluginInstalledFilename);
+		}
 
 		UGameFeaturesSubsystem::Get().OnGameFeatureActivating(StateProperties.GameFeatureData, StateProperties.PluginName, Context, StateProperties.PluginIdentifier.GetFullPluginURL());
 
@@ -2761,10 +2792,11 @@ void FInstallBundlePluginProtocolMetaData::ResetToDefaults()
 	VersionNum = FDefaultValues::CurrentVersionNum;
 	bUninstallBeforeTerminate = FDefaultValues::Default_bUninstallBeforeTerminate;
 	bUserPauseDownload = FDefaultValues::Default_bUserPauseDownload;
+	bAllowIniLoading = FDefaultValues::Default_bAllowIniLoading;
 	InstallBundleFlags = FDefaultValues::Default_InstallBundleFlags;
 	ReleaseInstallBundleFlags = FDefaultValues::Default_ReleaseInstallBundleFlags;
 
-	static_assert(static_cast<uint8>(EGameFeatureInstallBundleProtocolOptions::Count) == 6, "Update this function to handle the newly added EGameFeatureInstallBundleProtocolOptions value!");
+	static_assert(static_cast<uint8>(EGameFeatureInstallBundleProtocolOptions::Count) == 7, "Update this function to handle the newly added EGameFeatureInstallBundleProtocolOptions value!");
 }
 
 FString FInstallBundlePluginProtocolMetaData::ToString() const
@@ -2804,7 +2836,12 @@ FString FInstallBundlePluginProtocolMetaData::ToString() const
 		ReturnedString.Append(FString::Printf(TEXT("%s%s%s%s"), UE::GameFeatures::PluginURLStructureInfo::OptionSeperator, *LexToString(EGameFeatureInstallBundleProtocolOptions::UserPauseDownload), UE::GameFeatures::PluginURLStructureInfo::OptionAssignOperator, *LexToString(bUserPauseDownload)));
 	}
 
-	static_assert(static_cast<uint8>(EGameFeatureInstallBundleProtocolOptions::Count) == 6, "Update this function to handle the newly added EGameFeatureInstallBundleProtocolOptions value!");
+	if (bAllowIniLoading != FDefaultValues::Default_bAllowIniLoading)
+	{
+		ReturnedString.Append(FString::Printf(TEXT("%s%s%s%s"), UE::GameFeatures::PluginURLStructureInfo::OptionSeperator, *LexToString(EGameFeatureInstallBundleProtocolOptions::AllowIniLoading), UE::GameFeatures::PluginURLStructureInfo::OptionAssignOperator, *LexToString(bAllowIniLoading)));
+	}
+
+	static_assert(static_cast<uint8>(EGameFeatureInstallBundleProtocolOptions::Count) == 7, "Update this function to handle the newly added EGameFeatureInstallBundleProtocolOptions value!");
 
 	return ReturnedString;
 }
@@ -2939,6 +2976,21 @@ bool FInstallBundlePluginProtocolMetaData::FromString(const FString& URLString,
 						break;
 					}
 
+					case EGameFeatureInstallBundleProtocolOptions::AllowIniLoading:
+					{
+						if (OptionStrings.Num() != 2)
+						{
+							bParseSuccess = false;
+							UE_LOG(LogGameFeatures, Warning, TEXT("Error parsing InstallBundle protocol options URL %s . Invalid AllowIniLoading Option!"), *URLString);
+						}
+						else
+						{
+							Metadata.bAllowIniLoading = OptionStrings[1].ToBool();
+						}
+
+						break;
+					}
+
 					case EGameFeatureInstallBundleProtocolOptions::Count:
 					default:
 					{
@@ -2958,7 +3010,7 @@ bool FInstallBundlePluginProtocolMetaData::FromString(const FString& URLString,
 		UE_LOG(LogGameFeatures, Error, TEXT("Error parsing InstallBundle protocol options URL %s . No Bundle List Found!"), *URLString);
 	}
 
-	static_assert(static_cast<uint8>(EGameFeatureInstallBundleProtocolOptions::Count) == 6, "Update this function to handle the newly added EGameFeatureInstallBundleProtocolOptions value!");
+	static_assert(static_cast<uint8>(EGameFeatureInstallBundleProtocolOptions::Count) == 7, "Update this function to handle the newly added EGameFeatureInstallBundleProtocolOptions value!");
 
 	return bParseSuccess;
 }
@@ -3018,7 +3070,12 @@ bool FGameFeaturePluginStateMachineProperties::ValidateURLUpdate(const FGameFeat
 		const FInstallBundlePluginProtocolMetaData& NewMetaData = ProtocolMetadata.GetSubtype<FInstallBundlePluginProtocolMetaData>();
 		const FInstallBundlePluginProtocolMetaData& OldMetaData = OldProperties.ProtocolMetadata.GetSubtype<FInstallBundlePluginProtocolMetaData>();
 
-		if (!ensureMsgf((NewMetaData.InstallBundles == OldMetaData.InstallBundles), TEXT("Unexpected change in InstallBundles when updating URL: %s to %s"), *PluginIdentifier.GetFullPluginURL(), *OldProperties.PluginIdentifier.GetFullPluginURL()))
+		if (!ensureMsgf(NewMetaData.InstallBundles == OldMetaData.InstallBundles, TEXT("Unexpected change in InstallBundles when updating URL: %s to %s"), *PluginIdentifier.GetFullPluginURL(), *OldProperties.PluginIdentifier.GetFullPluginURL()))
+		{
+			bIsValidUpdate = false;
+		}
+
+		if (!ensureMsgf(NewMetaData.bAllowIniLoading == OldMetaData.bAllowIniLoading, TEXT("Unexpected change to AllowIniLoading when updating URL: %s to %s"), *PluginIdentifier.GetFullPluginURL(), *OldProperties.PluginIdentifier.GetFullPluginURL()))
 		{
 			bIsValidUpdate = false;
 		}
diff --git a/Engine/Plugins/Experimental/GameFeatures/Source/GameFeatures/Private/GameFeaturePluginStateMachine.h b/Engine/Plugins/Experimental/GameFeatures/Source/GameFeatures/Private/GameFeaturePluginStateMachine.h
--- a/Engine/Plugins/Experimental/GameFeatures/Source/GameFeatures/Private/GameFeaturePluginStateMachine.h
+++ b/Engine/Plugins/Experimental/GameFeatures/Source/GameFeatures/Private/GameFeaturePluginStateMachine.h
@@ -364,6 +364,8 @@ protected:
 	    returns False if it should just go directly to the Terminal state instead. */
 	bool ShouldVisitUninstallStateBeforeTerminal() const;
 
+	bool AllowIniLoading() const;
+
 private:
 	void CleanupDeferredUpdateCallbacks() const;
 
diff --git a/Engine/Plugins/Experimental/GameFeatures/Source/GameFeatures/Private/GameFeaturesSubsystem.cpp b/Engine/Plugins/Experimental/GameFeatures/Source/GameFeatures/Private/GameFeaturesSubsystem.cpp
--- a/Engine/Plugins/Experimental/GameFeatures/Source/GameFeatures/Private/GameFeaturesSubsystem.cpp
+++ b/Engine/Plugins/Experimental/GameFeatures/Source/GameFeatures/Private/GameFeaturesSubsystem.cpp
@@ -28,6 +28,7 @@ const uint32 FInstallBundlePluginProtocolMetaData::FDefaultValues::CurrentVersio
 //Missing InstallBundles on purpose as the default is just an empty TArray and should always be encoded
 const bool FInstallBundlePluginProtocolMetaData::FDefaultValues::Default_bUninstallBeforeTerminate = true;
 const bool FInstallBundlePluginProtocolMetaData::FDefaultValues::Default_bUserPauseDownload = false;
+const bool FInstallBundlePluginProtocolMetaData::FDefaultValues::Default_bAllowIniLoading = false;
 const EInstallBundleRequestFlags FInstallBundlePluginProtocolMetaData::FDefaultValues::Default_InstallBundleFlags = EInstallBundleRequestFlags::Defaults;
 const EInstallBundleReleaseRequestFlags FInstallBundlePluginProtocolMetaData::FDefaultValues::Default_ReleaseInstallBundleFlags = EInstallBundleReleaseRequestFlags::None;
 
diff --git a/Engine/Plugins/Experimental/GameFeatures/Source/GameFeatures/Public/GameFeaturesSubsystem.h b/Engine/Plugins/Experimental/GameFeatures/Source/GameFeatures/Public/GameFeaturesSubsystem.h
--- a/Engine/Plugins/Experimental/GameFeatures/Source/GameFeatures/Public/GameFeaturesSubsystem.h
+++ b/Engine/Plugins/Experimental/GameFeatures/Source/GameFeatures/Public/GameFeaturesSubsystem.h
@@ -281,6 +281,9 @@ struct GAMEFEATURES_API FInstallBundlePluginProtocolMetaData
 	/** If we want to set the Downloading state to pause because of user interaction */
 	bool bUserPauseDownload;
 
+	/** Allow the GFP to load INI files, should only be allowed for trusted content */
+	bool bAllowIniLoading;
+
 	/** Functions to convert to/from the URL FString representation of this metadata **/
 	FString ToString() const;
 	static bool FromString(const FString& URLString, FInstallBundlePluginProtocolMetaData& OutMetadata);
@@ -298,6 +301,7 @@ private:
 		//Missing InstallBundles on purpose as the default is just an empty TArray and should always be encoded
 		static const bool Default_bUninstallBeforeTerminate;
 		static const bool Default_bUserPauseDownload;
+		static const bool Default_bAllowIniLoading;
 		static const EInstallBundleRequestFlags Default_InstallBundleFlags;
 		static const EInstallBundleReleaseRequestFlags Default_ReleaseInstallBundleFlags;
 	};
