[Backout] - CL24105539
Original CL Desc
-----------------------------------------------------------------
Plugin Auditing - Adding an experimental plugin for auditing plugins to ensure that all the plugins the project is planning to cook don't depend on uncooked plug-in content.  Also checks sources for things like gameplay tags.

diff --git a/Engine/Plugins/Experimental/GameFeatures/Source/GameFeatures/Private/GameFeaturesProjectPolicies.cpp b/Engine/Plugins/Experimental/GameFeatures/Source/GameFeatures/Private/GameFeaturesProjectPolicies.cpp
--- a/Engine/Plugins/Experimental/GameFeatures/Source/GameFeatures/Private/GameFeaturesProjectPolicies.cpp
+++ b/Engine/Plugins/Experimental/GameFeatures/Source/GameFeatures/Private/GameFeaturesProjectPolicies.cpp
@@ -42,11 +42,6 @@ const TArray<FName> UDefaultGameFeaturesProjectPolicies::GetPreloadBundleStateFo
 	return FeatureBundles;
 }
 
-bool UGameFeaturesProjectPolicies::WillPluginBeCooked(const FString& PluginFilename, const FGameFeaturePluginDetails& PluginDetails) const
-{
-	return true;
-}
-
 void UGameFeaturesProjectPolicies::ExplicitLoadGameFeaturePlugin(const FString& PluginURL, const FGameFeaturePluginLoadComplete& CompleteDelegate, const bool bActivateGameFeatures)
 {
 	if (bActivateGameFeatures)
diff --git a/Engine/Plugins/Experimental/GameFeatures/Source/GameFeatures/Private/GameFeaturesSubsystem.cpp b/Engine/Plugins/Experimental/GameFeatures/Source/GameFeatures/Private/GameFeaturesSubsystem.cpp
--- a/Engine/Plugins/Experimental/GameFeatures/Source/GameFeatures/Private/GameFeaturesSubsystem.cpp
+++ b/Engine/Plugins/Experimental/GameFeatures/Source/GameFeatures/Private/GameFeaturesSubsystem.cpp
@@ -1145,26 +1145,42 @@ void UGameFeaturesSubsystem::LoadBuiltInGameFeaturePlugin(const TSharedRef<IPlug
 
 	UAssetManager::Get().PushBulkScanning();
 
-	FString PluginURL;
-	FGameFeaturePluginDetails PluginDetails;
-	if (GetGameFeaturePluginDetails(Plugin, PluginURL, PluginDetails))
+	const FString& PluginDescriptorFilename = Plugin->GetDescriptorFileName();
+
+	// Make sure you are in a game feature plugins folder. All GameFeaturePlugins are rooted in a GameFeatures folder.
+	if (!PluginDescriptorFilename.IsEmpty() && GetDefault<UGameFeaturesSubsystemSettings>()->IsValidGameFeaturePlugin(FPaths::ConvertRelativePathToFull(PluginDescriptorFilename)) && FPaths::FileExists(PluginDescriptorFilename))
 	{
-		if (GameSpecificPolicies->IsPluginAllowed(PluginURL))
+		FString PluginURL;
+		bool bIsFileProtocol = true;
+		if (GetPluginURLByName(Plugin->GetName(), PluginURL))
 		{
-			FBuiltInGameFeaturePluginBehaviorOptions BehaviorOptions;
-			const bool bShouldProcess = AdditionalFilter(Plugin->GetDescriptorFileName(), PluginDetails, BehaviorOptions);
-			if (bShouldProcess)
+			bIsFileProtocol = UGameFeaturesSubsystem::IsPluginURLProtocol(PluginURL, EGameFeaturePluginProtocol::File);
+		}
+		else
+		{
+			PluginURL = GetPluginURL_FileProtocol(PluginDescriptorFilename);
+		}
+
+		if (bIsFileProtocol && GameSpecificPolicies->IsPluginAllowed(PluginURL))
+		{
+			FGameFeaturePluginDetails PluginDetails;
+			if (GetGameFeaturePluginDetails(PluginDescriptorFilename, PluginDetails))
 			{
-				UGameFeaturePluginStateMachine* StateMachine = FindOrCreateGameFeaturePluginStateMachine(PluginURL);
+				FBuiltInGameFeaturePluginBehaviorOptions BehaviorOptions;
+				bool bShouldProcess = AdditionalFilter(PluginDescriptorFilename, PluginDetails, BehaviorOptions);
+				if (bShouldProcess)
+				{
+					UGameFeaturePluginStateMachine* StateMachine = FindOrCreateGameFeaturePluginStateMachine(PluginURL);
 
-				const EBuiltInAutoState InitialAutoState = (BehaviorOptions.AutoStateOverride != EBuiltInAutoState::Invalid) ? BehaviorOptions.AutoStateOverride : PluginDetails.BuiltInAutoState;
-				
-				const EGameFeaturePluginState DestinationState = ConvertInitialFeatureStateToTargetState(InitialAutoState);
+					const EBuiltInAutoState InitialAutoState = (BehaviorOptions.AutoStateOverride != EBuiltInAutoState::Invalid) ? BehaviorOptions.AutoStateOverride : PluginDetails.BuiltInAutoState;
+						
+					const EGameFeaturePluginState DestinationState = ConvertInitialFeatureStateToTargetState(InitialAutoState);
 
-				// If we're already at the destination or beyond, don't transition back
-				FGameFeaturePluginStateRange Destination(DestinationState, EGameFeaturePluginState::Active);
-				ChangeGameFeatureDestination(StateMachine, Destination, 
-					FGameFeaturePluginChangeStateComplete::CreateUObject(this, &ThisClass::LoadBuiltInGameFeaturePluginComplete, StateMachine, Destination));
+					// If we're already at the destination or beyond, don't transition back
+					FGameFeaturePluginStateRange Destination(DestinationState, EGameFeaturePluginState::Active);
+					ChangeGameFeatureDestination(StateMachine, Destination, 
+						FGameFeaturePluginChangeStateComplete::CreateUObject(this, &ThisClass::LoadBuiltInGameFeaturePluginComplete, StateMachine, Destination));
+				}
 			}
 		}
 	}
@@ -1336,32 +1352,6 @@ EGameFeaturePluginState UGameFeaturesSubsystem::GetPluginState(FGameFeaturePlugi
 	}
 }
 
-bool UGameFeaturesSubsystem::GetGameFeaturePluginDetails(const TSharedRef<IPlugin>& Plugin, FString& OutPluginURL, struct FGameFeaturePluginDetails& OutPluginDetails) const
-{
-	const FString& PluginDescriptorFilename = Plugin->GetDescriptorFileName();
-
-	// Make sure you are in a game feature plugins folder. All GameFeaturePlugins are rooted in a GameFeatures folder.
-	if (!PluginDescriptorFilename.IsEmpty() && GetDefault<UGameFeaturesSubsystemSettings>()->IsValidGameFeaturePlugin(FPaths::ConvertRelativePathToFull(PluginDescriptorFilename)) && FPaths::FileExists(PluginDescriptorFilename))
-	{
-		bool bIsFileProtocol = true;
-		if (GetPluginURLByName(Plugin->GetName(), OutPluginURL))
-		{
-			bIsFileProtocol = UGameFeaturesSubsystem::IsPluginURLProtocol(OutPluginURL, EGameFeaturePluginProtocol::File);
-		}
-		else
-		{
-			OutPluginURL = GetPluginURL_FileProtocol(PluginDescriptorFilename);
-		}
-
-		if (bIsFileProtocol)
-		{
-			return GetGameFeaturePluginDetails(PluginDescriptorFilename, OutPluginDetails);
-		}
-	}
-
-	return false;
-}
-
 bool UGameFeaturesSubsystem::GetGameFeaturePluginDetails(const FString& PluginDescriptorFilename, FGameFeaturePluginDetails& OutPluginDetails) const
 {
 	TRACE_CPUPROFILER_EVENT_SCOPE(GFP_GetPluginDetails);
diff --git a/Engine/Plugins/Experimental/GameFeatures/Source/GameFeatures/Public/GameFeaturesProjectPolicies.h b/Engine/Plugins/Experimental/GameFeatures/Source/GameFeatures/Public/GameFeaturesProjectPolicies.h
--- a/Engine/Plugins/Experimental/GameFeatures/Source/GameFeatures/Public/GameFeaturesProjectPolicies.h
+++ b/Engine/Plugins/Experimental/GameFeatures/Source/GameFeatures/Public/GameFeaturesProjectPolicies.h
@@ -23,9 +23,6 @@ public:
 	// Called when the game feature manager is shut down
 	virtual void ShutdownGameFeatureManager() { }
 
-	// Called to determined the expected state of a plugin under the WhenLoading conditions.
-	virtual bool WillPluginBeCooked(const FString& PluginFilename, const FGameFeaturePluginDetails& PluginDetails) const;
-
 	// Called when a game feature plugin enters the Loading state to determine additional assets to load
 	virtual TArray<FPrimaryAssetId> GetPreloadAssetListForGameFeature(const UGameFeatureData* GameFeatureToLoad, bool bIncludeLoadedAssets = false) const { return TArray<FPrimaryAssetId>(); }
 
diff --git a/Engine/Plugins/Experimental/GameFeatures/Source/GameFeatures/Public/GameFeaturesSubsystem.h b/Engine/Plugins/Experimental/GameFeatures/Source/GameFeatures/Public/GameFeaturesSubsystem.h
--- a/Engine/Plugins/Experimental/GameFeatures/Source/GameFeatures/Public/GameFeaturesSubsystem.h
+++ b/Engine/Plugins/Experimental/GameFeatures/Source/GameFeatures/Public/GameFeaturesSubsystem.h
@@ -496,8 +496,6 @@ public:
 	/** Returns the current state of the state machine for the specified plugin PluginIdentifier */
 	EGameFeaturePluginState GetPluginState(FGameFeaturePluginIdentifier PluginIdentifier) const;
 
-	/** Gets relevant properties out of a uplugin file */
-	bool GetGameFeaturePluginDetails(const TSharedRef<IPlugin>& Plugin, FString& OutPluginURL, struct FGameFeaturePluginDetails& OutPluginDetails) const;
 
 	/** Determine the initial feature state for a built-in plugin */
 	static EBuiltInAutoState DetermineBuiltInInitialFeatureState(TSharedPtr<FJsonObject> Descriptor, const FString& ErrorContext);
diff --git a/Engine/Plugins/Experimental/PluginAudit/PluginAudit.uplugin b/Engine/Plugins/Experimental/PluginAudit/PluginAudit.uplugin
deleted file mode 100644
diff --git a/Engine/Plugins/Experimental/PluginAudit/Source/PluginAudit.Build.cs b/Engine/Plugins/Experimental/PluginAudit/Source/PluginAudit.Build.cs
deleted file mode 100644
diff --git a/Engine/Plugins/Experimental/PluginAudit/Source/Private/PluginAuditModule.cpp b/Engine/Plugins/Experimental/PluginAudit/Source/Private/PluginAuditModule.cpp
deleted file mode 100644
diff --git a/Engine/Plugins/Experimental/PluginAudit/Source/Private/SPluginAuditBrowser.cpp b/Engine/Plugins/Experimental/PluginAudit/Source/Private/SPluginAuditBrowser.cpp
deleted file mode 100644
diff --git a/Engine/Plugins/Experimental/PluginAudit/Source/Private/SPluginAuditBrowser.h b/Engine/Plugins/Experimental/PluginAudit/Source/Private/SPluginAuditBrowser.h
deleted file mode 100644
