Mark assets as garbage before gc in plugin unload.

diff --git a/Engine/Plugins/Experimental/GameFeatures/Source/GameFeatures/Private/GameFeaturePluginStateMachine.cpp b/Engine/Plugins/Experimental/GameFeatures/Source/GameFeatures/Private/GameFeaturePluginStateMachine.cpp
--- a/Engine/Plugins/Experimental/GameFeatures/Source/GameFeatures/Private/GameFeaturePluginStateMachine.cpp
+++ b/Engine/Plugins/Experimental/GameFeatures/Source/GameFeatures/Private/GameFeaturePluginStateMachine.cpp
@@ -24,6 +24,7 @@
 #include "GameFeaturesProjectPolicies.h"
 #include "UObject/ObjectRename.h"
 #include "UObject/UObjectAllocator.h"
+#include "Misc/PathViews.h"
 
 #include UE_INLINE_GENERATED_CPP_BY_NAME(GameFeaturePluginStateMachine)
 
@@ -37,6 +38,12 @@ namespace UE::GameFeatures
 
 	static UE::GameFeatures::FResult CanceledResult = MakeError(StateMachineErrorNamespace + TEXT("Canceled"));
 
+
+	static bool MarkPluginAsGarbageOnUnload = true;
+	static FAutoConsoleVariableRef CVarMarkPluginAsGarbageOnUnload(TEXT("GameFeaturePlugin.MarkPluginAsGarbageOnUnload"),
+		MarkPluginAsGarbageOnUnload,
+		TEXT("Mark up the plugin as garbage when the plugin is unloaded (helps track down gc issues)."));
+
 	static int32 ShouldLogMountedFiles = 0;
 	static FAutoConsoleVariableRef CVarShouldLogMountedFiles(TEXT("GameFeaturePlugin.ShouldLogMountedFiles"),
 		ShouldLogMountedFiles,
@@ -2073,7 +2080,12 @@ struct FGameFeaturePluginState_Unloading : public FGameFeaturePluginState
 		}
 
 		UnloadGameFeatureBundles(StateProperties.GameFeatureData);
-
+#if !WITH_EDITOR // never mark anything as garbage in the editor, there be dwagons here. 
+		if (UE::GameFeatures::MarkPluginAsGarbageOnUnload)
+		{
+			MarkPluginAsGarbage(StateProperties.PluginName);
+		}
+#endif
 		if (StateProperties.Destination.MaxState == EGameFeaturePluginState::Registered)
 		{
 			// If we aren't going farther than Registered, GC now
@@ -2086,6 +2098,35 @@ struct FGameFeaturePluginState_Unloading : public FGameFeaturePluginState
 		StateStatus.SetTransition(EGameFeaturePluginState::Registered);
 	}
 
+	void MarkPluginAsGarbage(FString PluginName)
+	{
+		TSet<FString> PluginNames;
+		PluginNames.Add(PluginName);
+
+		TArray<UPackage*> PackagesToUnload;
+		for (TObjectIterator<UPackage> It; It; ++It)
+		{
+			const FNameBuilder PackageName(It->GetFName());
+			const FStringView PackageMountPointName = FPathViews::GetMountPointNameFromPath(PackageName);
+			if (PluginName == PackageMountPointName)
+			{
+				PackagesToUnload.Add(*It);
+			}
+		}
+		
+		for (UPackage* Package : PackagesToUnload)
+		{
+			UE_LOG(LogGameFeatures, Log, TEXT("Marking package %s as Garbage"), *Package->GetName());
+			ForEachObjectWithPackage(Package, [](UObject* Object)
+				{
+					Object->MarkAsGarbage();
+			return true;
+				}, false);
+
+			Package->MarkAsGarbage();
+		}
+	}
+
 	void UnloadGameFeatureBundles(const UGameFeatureData* GameFeatureToLoad)
 	{
 		if (GameFeatureToLoad == nullptr)
