Fix unity compiler errors

diff --git a/Engine/Plugins/Experimental/GameFeatures/Source/GameFeatures/Private/GameFeaturePluginStateMachine.cpp b/Engine/Plugins/Experimental/GameFeatures/Source/GameFeatures/Private/GameFeaturePluginStateMachine.cpp
--- a/Engine/Plugins/Experimental/GameFeatures/Source/GameFeatures/Private/GameFeaturePluginStateMachine.cpp
+++ b/Engine/Plugins/Experimental/GameFeatures/Source/GameFeatures/Private/GameFeaturePluginStateMachine.cpp
@@ -24,6 +24,8 @@
 #include "GameFeaturesProjectPolicies.h"
 #include "UObject/ObjectRename.h"
 #include "UObject/UObjectAllocator.h"
+#include "UObject/UObjectIterator.h"
+#include "UObject/NoExportTypes.h"
 #include "Misc/PathViews.h"
 
 #include UE_INLINE_GENERATED_CPP_BY_NAME(GameFeaturePluginStateMachine)
