Use skip verify list to also skip markup unloaded plugins as garbage.

diff --git a/Engine/Plugins/Experimental/GameFeatures/Source/GameFeatures/Private/GameFeaturePluginStateMachine.cpp b/Engine/Plugins/Experimental/GameFeatures/Source/GameFeatures/Private/GameFeaturePluginStateMachine.cpp
--- a/Engine/Plugins/Experimental/GameFeatures/Source/GameFeatures/Private/GameFeaturePluginStateMachine.cpp
+++ b/Engine/Plugins/Experimental/GameFeatures/Source/GameFeatures/Private/GameFeaturePluginStateMachine.cpp
@@ -115,6 +115,29 @@ namespace UE::GameFeatures
 	}
 	#undef GAME_FEATURE_PLUGIN_STATE_TO_STRING
 
+	bool ShouldSkipVerify(const FString& PluginName)
+	{
+
+		static const FAsciiSet Wildcards("*?");
+		bool bSkip = false;
+		UE::String::ParseTokens(MakeStringView(GVerifyPluginSkipList), TEXTVIEW(","), [&PluginName, &bSkip](FStringView Item) {
+				if (bSkip) { return; }
+				if (Item.Equals(PluginName, ESearchCase::IgnoreCase))
+				{
+					bSkip = true;
+				}
+				else if (FAsciiSet::HasAny(Item, Wildcards))
+				{
+					FString Pattern = FString(Item); // Need to copy to null terminate
+					if (FWildcardString::IsMatchSubstring(*Pattern, *PluginName, *PluginName + PluginName.Len(), ESearchCase::IgnoreCase))
+					{
+						bSkip = true;
+					}
+				}
+			});
+		return bSkip;
+	}
+
 	// Check if any assets from the plugin mount point have leaked, and if so trace them.
 	// Then rename them to allow new copies of them to be loaded. 
 	void HandlePossibleAssetLeaks(const FString& PluginName, UPackage* IgnorePackage = nullptr)
@@ -127,25 +150,8 @@ namespace UE::GameFeatures
 			return;
 		}
 
-		static const FAsciiSet Wildcards("*?");
-		bool bSkip = false;
-		UE::String::ParseTokens(MakeStringView(GVerifyPluginSkipList), TEXTVIEW(","), [&PluginName, &bSkip](FStringView Item) {
-			if (bSkip) { return; }
-			if (Item.Equals(PluginName, ESearchCase::IgnoreCase))
-			{
-				bSkip = true;
-			}
-			else if (FAsciiSet::HasAny(Item, Wildcards))
-			{
-				FString Pattern = FString(Item); // Need to copy to null terminate
-				if (FWildcardString::IsMatchSubstring(*Pattern, *PluginName, *PluginName + PluginName.Len(), ESearchCase::IgnoreCase))
-				{
-					bSkip = true;
-				}
-			}
-		});
 
-		if (bSkip)
+		if (UE::GameFeatures::ShouldSkipVerify(PluginName))
 		{
 			return;
 		}
@@ -2083,7 +2089,7 @@ struct FGameFeaturePluginState_Unloading : public FGameFeaturePluginState
 
 		UnloadGameFeatureBundles(StateProperties.GameFeatureData);
 #if !WITH_EDITOR // never mark anything as garbage in the editor, there be dwagons here. 
-		if (UE::GameFeatures::MarkPluginAsGarbageOnUnload)
+		if (UE::GameFeatures::MarkPluginAsGarbageOnUnload && !UE::GameFeatures::ShouldSkipVerify(StateProperties.PluginName))
 		{
 			MarkPluginAsGarbage(StateProperties.PluginName);
 		}
