Fixed game feature data from being cleaned up when transitioning to unloaded.

diff --git a/Engine/Plugins/Experimental/GameFeatures/Source/GameFeatures/Private/GameFeaturePluginStateMachine.cpp b/Engine/Plugins/Experimental/GameFeatures/Source/GameFeatures/Private/GameFeaturePluginStateMachine.cpp
--- a/Engine/Plugins/Experimental/GameFeatures/Source/GameFeatures/Private/GameFeaturePluginStateMachine.cpp
+++ b/Engine/Plugins/Experimental/GameFeatures/Source/GameFeatures/Private/GameFeaturePluginStateMachine.cpp
@@ -46,6 +46,11 @@ namespace UE::GameFeatures
 		MarkPluginAsGarbageOnUnload,
 		TEXT("Mark up the plugin as garbage when the plugin is unloaded (helps track down gc issues)."));
 
+	static bool MarkPluginAsGarbageOnUnregister = false;
+	static FAutoConsoleVariableRef CVarMarkPluginAsGarbageOnUnregister(TEXT("GameFeaturePlugin.MarkPluginAsGarbageOnUnregister"),
+		MarkPluginAsGarbageOnUnregister,
+		TEXT("Mark up the plugin as garbage when the plugin is unloaded (helps track down gc issues)."));
+
 	static int32 ShouldLogMountedFiles = 0;
 	static FAutoConsoleVariableRef CVarShouldLogMountedFiles(TEXT("GameFeaturePlugin.ShouldLogMountedFiles"),
 		ShouldLogMountedFiles,
@@ -544,6 +549,42 @@ void FGameFeaturePluginState::GarbageCollectAndUpdateStateMachineDeferred() cons
 	FCoreUObjectDelegates::GetPostGarbageCollect().AddRaw(this, &FGameFeaturePluginState::UpdateStateMachineDeferred, 0.0f);
 }
 
+
+void FGameFeaturePluginState::MarkPluginAsGarbage(bool bMarkGameFeatureDataAsGarbage)
+{
+	TArray<UPackage*> PackagesToUnload;
+	for (TObjectIterator<UPackage> It; It; ++It)
+	{
+		const FNameBuilder PackageName(It->GetFName());
+		const FStringView PackageMountPointName = FPathViews::GetMountPointNameFromPath(PackageName);
+		if (StateProperties.PluginName == PackageMountPointName)
+		{
+			PackagesToUnload.Add(*It);
+		}
+	}
+
+	for (UPackage* Package : PackagesToUnload)
+	{
+		bool bShouldKeep = false;
+		ForEachObjectWithPackage(Package, [&bShouldKeep](UObject* Object)
+			{
+				if (Object->IsA(UGameFeatureData::StaticClass()))
+				{
+					bShouldKeep = true;
+					return true;
+				}
+		Object->MarkAsGarbage();
+		return true;
+			}, false);
+		if (!bShouldKeep)
+		{
+			UE_LOG(LogGameFeatures, Log, TEXT("Marking package %s as Garbage"), *Package->GetName());
+			Package->MarkAsGarbage();
+		}
+
+	}
+}
+
 void FGameFeaturePluginState::UpdateStateMachineImmediate() const
 {
 	StateProperties.OnRequestUpdateStateMachine.ExecuteIfBound();
@@ -1958,6 +1999,13 @@ struct FGameFeaturePluginState_Unregistering : public FGameFeaturePluginState
 		verify(FPluginUtils::UnloadPluginAssets(StateProperties.PluginName));
 #endif //if WITH_EDITOR
 
+#if !WITH_EDITOR // never mark anything as garbage in the editor, there be dwagons here. 
+		if (UE::GameFeatures::MarkPluginAsGarbageOnUnregister && !UE::GameFeatures::ShouldSkipVerify(StateProperties.PluginName))
+		{
+			MarkPluginAsGarbage(true);
+		}
+#endif
+
 		// Try to remove the gameplay tags, this might be ignored depending on project settings
 		const FString PluginFolder = FPaths::GetPath(StateProperties.PluginInstalledFilename);
 		UGameplayTagsManager::Get().RemoveTagIniSearchPath(PluginFolder / TEXT("Config") / TEXT("Tags"));
@@ -2091,7 +2139,7 @@ struct FGameFeaturePluginState_Unloading : public FGameFeaturePluginState
 #if !WITH_EDITOR // never mark anything as garbage in the editor, there be dwagons here. 
 		if (UE::GameFeatures::MarkPluginAsGarbageOnUnload && !UE::GameFeatures::ShouldSkipVerify(StateProperties.PluginName))
 		{
-			MarkPluginAsGarbage(StateProperties.PluginName);
+			MarkPluginAsGarbage(false);
 		}
 #endif
 		if (StateProperties.Destination.MaxState == EGameFeaturePluginState::Registered)
@@ -2106,34 +2154,6 @@ struct FGameFeaturePluginState_Unloading : public FGameFeaturePluginState
 		StateStatus.SetTransition(EGameFeaturePluginState::Registered);
 	}
 
-	void MarkPluginAsGarbage(FString PluginName)
-	{
-		TSet<FString> PluginNames;
-		PluginNames.Add(PluginName);
-
-		TArray<UPackage*> PackagesToUnload;
-		for (TObjectIterator<UPackage> It; It; ++It)
-		{
-			const FNameBuilder PackageName(It->GetFName());
-			const FStringView PackageMountPointName = FPathViews::GetMountPointNameFromPath(PackageName);
-			if (PluginName == PackageMountPointName)
-			{
-				PackagesToUnload.Add(*It);
-			}
-		}
-		
-		for (UPackage* Package : PackagesToUnload)
-		{
-			UE_LOG(LogGameFeatures, Log, TEXT("Marking package %s as Garbage"), *Package->GetName());
-			ForEachObjectWithPackage(Package, [](UObject* Object)
-				{
-					Object->MarkAsGarbage();
-			return true;
-				}, false);
-
-			Package->MarkAsGarbage();
-		}
-	}
 
 	void UnloadGameFeatureBundles(const UGameFeatureData* GameFeatureToLoad)
 	{
diff --git a/Engine/Plugins/Experimental/GameFeatures/Source/GameFeatures/Private/GameFeaturePluginStateMachine.h b/Engine/Plugins/Experimental/GameFeatures/Source/GameFeatures/Private/GameFeaturePluginStateMachine.h
--- a/Engine/Plugins/Experimental/GameFeatures/Source/GameFeatures/Private/GameFeaturePluginStateMachine.h
+++ b/Engine/Plugins/Experimental/GameFeatures/Source/GameFeatures/Private/GameFeaturePluginStateMachine.h
@@ -352,6 +352,9 @@ struct FGameFeaturePluginState
 
 	void UpdateProgress(float Progress) const;
 
+
+	void MarkPluginAsGarbage(bool bMarkGameFeatureDataAsGarbage);
+
 protected:
 	/** Builds an end FResult with some minimal error information with overrides for common types we 
 		need to generate errors from */
