Fix for mark gpf as garbage code not respecting marking game feature data as garbage flag.

diff --git a/Engine/Plugins/Experimental/GameFeatures/Source/GameFeatures/Private/GameFeaturePluginStateMachine.cpp b/Engine/Plugins/Experimental/GameFeatures/Source/GameFeatures/Private/GameFeaturePluginStateMachine.cpp
--- a/Engine/Plugins/Experimental/GameFeatures/Source/GameFeatures/Private/GameFeaturePluginStateMachine.cpp
+++ b/Engine/Plugins/Experimental/GameFeatures/Source/GameFeatures/Private/GameFeaturePluginStateMachine.cpp
@@ -566,15 +566,15 @@ void FGameFeaturePluginState::MarkPluginAsGarbage(bool bMarkGameFeatureDataAsGar
 	for (UPackage* Package : PackagesToUnload)
 	{
 		bool bShouldKeep = false;
-		ForEachObjectWithPackage(Package, [&bShouldKeep](UObject* Object)
+		ForEachObjectWithPackage(Package, [&bShouldKeep, bMarkGameFeatureDataAsGarbage](UObject* Object)
 			{
-				if (Object->IsA(UGameFeatureData::StaticClass()))
+				if (Object->IsA(UGameFeatureData::StaticClass()) && !bMarkGameFeatureDataAsGarbage)
 				{
 					bShouldKeep = true;
 					return true;
 				}
-		Object->MarkAsGarbage();
-		return true;
+				Object->MarkAsGarbage();
+				return true;
 			}, false);
 		if (!bShouldKeep)
 		{
