Limit the number of objects we search for references for during GFP leak detection.
Only search for packages (assets will have the package in their outer chain, further reduces number of objects queried).
When no object passed to FindAndPrintStaleReferencesToObjects have the Garbage flag, use the Shortest search mode rather than ShortestToGarbage.

diff --git a/Engine/Plugins/Experimental/GameFeatures/Source/GameFeatures/Private/GameFeaturePluginStateMachine.cpp b/Engine/Plugins/Experimental/GameFeatures/Source/GameFeatures/Private/GameFeaturePluginStateMachine.cpp
--- a/Engine/Plugins/Experimental/GameFeatures/Source/GameFeatures/Private/GameFeaturePluginStateMachine.cpp
+++ b/Engine/Plugins/Experimental/GameFeatures/Source/GameFeatures/Private/GameFeaturePluginStateMachine.cpp
@@ -11,6 +11,7 @@
 #include "InstallBundleUtils.h"
 #include "BundlePrereqCombinedStatusHelper.h"
 #include "Interfaces/IPluginManager.h"
+#include "Materials/MaterialInterface.h"
 #include "Misc/AsciiSet.h"
 #include "Misc/ConfigCacheIni.h"
 #include "Misc/CoreDelegates.h"
@@ -107,6 +108,14 @@ namespace UE::GameFeatures
 		ECVF_Default
 	);
 
+	static int32 GLeakedAssetTrace_MaxReportCount = 10;
+	static FAutoConsoleVariableRef CVarLeakedAssetTrace_MaxReportCount(
+		TEXT("GameFeaturePlugin.LeakedAssetTrace.MaxReportCount"),
+		GLeakedAssetTrace_MaxReportCount,
+		TEXT("Max number of assets to report when we find leaked assets.\n"),
+		ECVF_Default
+	);
+
 	#define GAME_FEATURE_PLUGIN_STATE_TO_STRING(inEnum, inText) case EGameFeaturePluginState::inEnum: return TEXT(#inEnum);
 	FString ToString(EGameFeaturePluginState InType)
 	{
@@ -143,6 +152,27 @@ namespace UE::GameFeatures
 		return bSkip;
 	}
 
+	// Return a higher number for packages which it's more important to include in leak reporting, when the number of leaks we want to report is limited. 
+	int32 GetPackageLeakReportingPriority(UPackage* Package)
+	{	
+		int32 Priority = 0;
+		ForEachObjectWithPackage(Package, [&Priority](UObject* Object)
+		{
+			if (UWorld* World = Cast<UWorld>(Object))
+			{
+				Priority = 100;
+				return true;
+			}
+			else if (Cast<UMaterialInterface>(Object))
+			{
+				Priority = FMath::Max(Priority, 50);
+				// keep iterating in case we find a world 
+			}
+			return true;
+		}, false);
+		return Priority;
+	}
+
 	// Check if any assets from the plugin mount point have leaked, and if so trace them.
 	// Then rename them to allow new copies of them to be loaded. 
 	void HandlePossibleAssetLeaks(const FString& PluginName, UPackage* IgnorePackage = nullptr)
@@ -161,7 +191,7 @@ namespace UE::GameFeatures
 			return;
 		}
 		
-		TSet<UPackage*> LeakedPackages;
+		TMap<UPackage*, int32> LeakedPackages;
 		TStringBuilder<512> Prefix;
 		Prefix << '/' << PluginName << '/';
 
@@ -186,7 +216,7 @@ namespace UE::GameFeatures
 					Package->GetFName().GetDisplayNameEntry()->AppendNameToString(NameBuffer);
 					if (NameBuffer.ToView().StartsWith(Prefix, ESearchCase::IgnoreCase))
 					{
-						LeakedPackages.Add(Package);
+						LeakedPackages.Add(Package, GetPackageLeakReportingPriority(Package));
 					}
 				}
 			});
@@ -219,32 +249,27 @@ namespace UE::GameFeatures
 				Options |= EPrintStaleReferencesOptions::Minimal;
 			}
 
-			// To minimize size of log, try to search for just public objects from the leaked packages.
-			TArray<UObject*> LeakedObjects;
-			for (UPackage* Package : LeakedPackages)
+			// Sort even if we don't limit the count, so that high priority leaks appear first 
+			LeakedPackages.ValueSort(TGreater<int32>{});
+			int32 OmittedCount = FMath::Max(0, LeakedPackages.Num() - GLeakedAssetTrace_MaxReportCount);
+			TArray<UPackage*> PackagesToSearchFor;
+			int32 i = 0;
+			for (const TPair<UPackage*, int32>& Pair : LeakedPackages)
 			{
-				int32 StartCount = LeakedObjects.Num();
-				ForEachObjectWithPackage(Package, [&LeakedObjects](UObject* Object)
-				{
-					if (Object->HasAnyFlags(RF_Public))
-					{
-						LeakedObjects.Add(Object);
-					}
-					return true;
-				}, false);
-				if (LeakedObjects.Num() == StartCount)
+				if (i++ < GLeakedAssetTrace_MaxReportCount)
 				{
-					LeakedObjects.Add(Package);
+					PackagesToSearchFor.Add(Pair.Key);
 				}
 			}
 
-			UE_LOG(LogGameFeatures, Display, TEXT("Searching for references to %d leaked objects from plugin %s"), LeakedObjects.Num(), *PluginName);
-			GEngine->FindAndPrintStaleReferencesToObjects(LeakedObjects, Options);
+			UE_LOG(LogGameFeatures, Display, TEXT("Searching for references to %d leaked packages (%d omitted for speed) from plugin %s"), LeakedPackages.Num(), OmittedCount, *PluginName);
+			GEngine->FindAndPrintStaleReferencesToObjects(MakeArrayView((UObject**)PackagesToSearchFor.GetData(), PackagesToSearchFor.Num()), Options);
 		}
 
 		// Rename the packages that we are streaming out so that we can possibly reload another copy of them
-		for (UPackage* Package : LeakedPackages)
+		for (const TPair<UPackage*, int32>& Pair : LeakedPackages)
 		{
+			UPackage* Package = Pair.Key;
 			UE_LOG(LogGameFeatures, Warning, TEXT("Marking leaking package %s as Garbage"), *Package->GetName());
 			ForEachObjectWithPackage(Package, [](UObject* Object)
 			{
diff --git a/Engine/Source/Runtime/Engine/Private/UnrealEngine.cpp b/Engine/Source/Runtime/Engine/Private/UnrealEngine.cpp
--- a/Engine/Source/Runtime/Engine/Private/UnrealEngine.cpp
+++ b/Engine/Source/Runtime/Engine/Private/UnrealEngine.cpp
@@ -16064,9 +16064,11 @@ TArray<FString> UEngine::FindAndPrintStaleReferencesToObjects(TConstArrayView<UO
 {
 	TArray<FString> Paths; 
 	checkf(ObjectsToFindReferencesTo.Num() > 0, TEXT("ObjectsToFindReferencesTo cannot be empty"));
+	bool bAnyGarbage = false;
 	for (UObject* Obj : ObjectsToFindReferencesTo)
 	{
 		checkf(Obj, TEXT("Object to find references to cannot be null"));
+		bAnyGarbage = bAnyGarbage || Obj->HasAnyInternalFlags(EInternalObjectFlags::Garbage);
 	}
 
 	ELogVerbosity::Type Verbosity = ELogVerbosity::NoLogging;
@@ -16104,7 +16106,7 @@ TArray<FString> UEngine::FindAndPrintStaleReferencesToObjects(TConstArrayView<UO
 	}
 	else
 	{
-		SearchMode |= EReferenceChainSearchMode::ShortestToGarbage;
+		SearchMode |= bAnyGarbage ? EReferenceChainSearchMode::ShortestToGarbage : EReferenceChainSearchMode::Shortest;
 	}
 
 	FReferenceChainSearch RefChainSearch(ObjectsToFindReferencesTo, SearchMode);
