Fix for state machine advancing through unregistration before all of the neceesary objects have been cleaned up.

diff --git a/Engine/Plugins/Experimental/GameFeatures/Source/GameFeatures/GameFeatures.Build.cs b/Engine/Plugins/Experimental/GameFeatures/Source/GameFeatures/GameFeatures.Build.cs
--- a/Engine/Plugins/Experimental/GameFeatures/Source/GameFeatures/GameFeatures.Build.cs
+++ b/Engine/Plugins/Experimental/GameFeatures/Source/GameFeatures/GameFeatures.Build.cs
@@ -27,6 +27,7 @@ namespace UnrealBuildTool.Rules
 					"Json",
 					"PakFile",
 					"Projects",
+					"RenderCore", // required for FDeferredCleanupInterface
 				}
 			);
 
diff --git a/Engine/Plugins/Experimental/GameFeatures/Source/GameFeatures/Private/GameFeaturePluginStateMachine.cpp b/Engine/Plugins/Experimental/GameFeatures/Source/GameFeatures/Private/GameFeaturePluginStateMachine.cpp
--- a/Engine/Plugins/Experimental/GameFeatures/Source/GameFeatures/Private/GameFeaturePluginStateMachine.cpp
+++ b/Engine/Plugins/Experimental/GameFeatures/Source/GameFeatures/Private/GameFeaturePluginStateMachine.cpp
@@ -566,14 +566,43 @@ void FGameFeaturePluginState::UpdateStateMachineDeferred(float Delay /*= 0.0f*/)
 	}), Delay);
 }
 
-void FGameFeaturePluginState::GarbageCollectAndUpdateStateMachineDeferred() const
+void FGameFeaturePluginState::InsertCleanupObjectsFence() const
+{
+	class FLevelStreamoutFence : public FDeferredCleanupInterface
+	{
+	public:
+		FLevelStreamoutFence(const FGameFeaturePluginState& InState)
+			: State(InState)
+		{
+
+		}
+		virtual ~FLevelStreamoutFence()
+		{
+			State.UpdateStateMachineDeferred();
+		}
+
+	private:
+		const FGameFeaturePluginState& State;
+	};
+
+	BeginCleanup(new FLevelStreamoutFence(*this));
+}
+
+void FGameFeaturePluginState::GarbageCollectAndUpdateStateMachineDeferred(bool bWaitForCleanupQueue /*= false*/) const
 {
 	GEngine->ForceGarbageCollection(true); // Tick Delayed
 
 	CleanupDeferredUpdateCallbacks();
-	FCoreUObjectDelegates::GetPostGarbageCollect().AddRaw(this, &FGameFeaturePluginState::UpdateStateMachineDeferred, 0.0f);
-}
 
+	if (bWaitForCleanupQueue)
+	{
+		FCoreUObjectDelegates::GetPostGarbageCollect().AddRaw(this, &FGameFeaturePluginState::InsertCleanupObjectsFence);
+	}
+	else
+	{
+		FCoreUObjectDelegates::GetPostGarbageCollect().AddRaw(this, &FGameFeaturePluginState::UpdateStateMachineDeferred, 0.0f);
+	}
+}
 
 void FGameFeaturePluginState::MarkPluginAsGarbage(bool bMarkGameFeatureDataAsGarbage)
 {
@@ -2036,7 +2065,7 @@ struct FGameFeaturePluginState_Unregistering : public FGameFeaturePluginState
 		UGameplayTagsManager::Get().RemoveTagIniSearchPath(PluginFolder / TEXT("Config") / TEXT("Tags"));
 
 		bRequestedGC = true;
-		GarbageCollectAndUpdateStateMachineDeferred();
+		GarbageCollectAndUpdateStateMachineDeferred(true);
 	}
 };
 
diff --git a/Engine/Plugins/Experimental/GameFeatures/Source/GameFeatures/Private/GameFeaturePluginStateMachine.h b/Engine/Plugins/Experimental/GameFeatures/Source/GameFeatures/Private/GameFeaturePluginStateMachine.h
--- a/Engine/Plugins/Experimental/GameFeatures/Source/GameFeatures/Private/GameFeaturePluginStateMachine.h
+++ b/Engine/Plugins/Experimental/GameFeatures/Source/GameFeatures/Private/GameFeaturePluginStateMachine.h
@@ -346,8 +346,9 @@ struct FGameFeaturePluginState
 	/** The common properties that can be accessed by the states of the state machine */
 	FGameFeaturePluginStateMachineProperties& StateProperties;
 
+	void InsertCleanupObjectsFence() const;
 	void UpdateStateMachineDeferred(float Delay = 0.0f) const;
-	void GarbageCollectAndUpdateStateMachineDeferred() const;
+	void GarbageCollectAndUpdateStateMachineDeferred(bool bWaitForCleanupQueue = false) const;
 	void UpdateStateMachineImmediate() const;
 
 	void UpdateProgress(float Progress) const;
