non-unity build fix for GameFeaturePluginStateMachine.cpp (follow up for 24515739)
#rnx

diff --git a/Engine/Plugins/Experimental/GameFeatures/Source/GameFeatures/Private/GameFeaturePluginStateMachine.cpp b/Engine/Plugins/Experimental/GameFeatures/Source/GameFeatures/Private/GameFeaturePluginStateMachine.cpp
--- a/Engine/Plugins/Experimental/GameFeatures/Source/GameFeatures/Private/GameFeaturePluginStateMachine.cpp
+++ b/Engine/Plugins/Experimental/GameFeatures/Source/GameFeatures/Private/GameFeaturePluginStateMachine.cpp
@@ -20,6 +20,7 @@
 #include "Misc/WildcardString.h"
 #include "Algo/AllOf.h"
 #include "Misc/TVariantMeta.h"
+#include "RenderDeferredCleanup.h"
 #include "Serialization/MemoryReader.h"
 #include "String/ParseTokens.h"
 #include "GameFeaturesProjectPolicies.h"
