Game Feature Plugins
- Fix UninstallBeforeTerminate defaulting to true
-- This could case extra release requests for bundle sources that cached all their data and didn't explicitly need an uninstall and could lead to extra IO that didn't accomplish anything.
-- The only existing uninstall case dealing with non-LRU data explicitly calls UninstallGameFeaturePlugin. so chaning the default doesn't affect it
- Fix lambda delegates not being weakly bound

#localization none
#tests Tested locally, verified [redacted] bundle source only did one journal write on release

diff --git a/Engine/Plugins/Experimental/GameFeatures/Source/GameFeatures/Private/GameFeaturesSubsystem.cpp b/Engine/Plugins/Experimental/GameFeatures/Source/GameFeatures/Private/GameFeaturesSubsystem.cpp
--- a/Engine/Plugins/Experimental/GameFeatures/Source/GameFeatures/Private/GameFeaturesSubsystem.cpp
+++ b/Engine/Plugins/Experimental/GameFeatures/Source/GameFeatures/Private/GameFeaturesSubsystem.cpp
@@ -26,7 +26,7 @@ DEFINE_LOG_CATEGORY(LogGameFeatures);
 
 const uint32 FInstallBundlePluginProtocolMetaData::FDefaultValues::CurrentVersionNum = 1;
 //Missing InstallBundles on purpose as the default is just an empty TArray and should always be encoded
-const bool FInstallBundlePluginProtocolMetaData::FDefaultValues::Default_bUninstallBeforeTerminate = true;
+const bool FInstallBundlePluginProtocolMetaData::FDefaultValues::Default_bUninstallBeforeTerminate = false;
 const bool FInstallBundlePluginProtocolMetaData::FDefaultValues::Default_bUserPauseDownload = false;
 const bool FInstallBundlePluginProtocolMetaData::FDefaultValues::Default_bAllowIniLoading = false;
 const EInstallBundleRequestFlags FInstallBundlePluginProtocolMetaData::FDefaultValues::Default_InstallBundleFlags = EInstallBundleRequestFlags::Defaults;
@@ -1048,12 +1048,6 @@ void UGameFeaturesSubsystem::UninstallGameFeaturePlugin(const FString& PluginURL
 	UGameFeaturePluginStateMachine* StateMachine = FindOrCreateGameFeaturePluginStateMachine(PluginURL);
 	check(StateMachine);
 
-	//Define a lambda wrapper to return the FResult of the Terminate for our FGameFeaturePluginUninstallComplete Delegate
-	const FGameFeaturePluginTerminateComplete ReportTerminateResultLambda = FGameFeaturePluginTerminateComplete::CreateLambda([=](const UE::GameFeatures::FResult& Result)
-		{
-			CompleteDelegate.Execute(Result);
-		});
-
 	//We may need to update our PluginURL to force certain metadata changes to facilitate this uninstall
 	//This tracks what URL we actually will pass in to the terminate as we can rely on the fact that Terminate
 	//will update any important Metadata from this new URL before beginning it's terminate
@@ -1077,7 +1071,7 @@ void UGameFeaturesSubsystem::UninstallGameFeaturePlugin(const FString& PluginURL
 		}
 	}
 
-	TerminateGameFeaturePlugin(PluginURLForTerminate, ReportTerminateResultLambda);
+	TerminateGameFeaturePlugin(PluginURLForTerminate, CompleteDelegate);
 }
 
 void UGameFeaturesSubsystem::TerminateGameFeaturePlugin(const FString& PluginURL)
@@ -1090,7 +1084,8 @@ void UGameFeaturesSubsystem::TerminateGameFeaturePlugin(const FString& PluginURL
 	if (UGameFeaturePluginStateMachine* StateMachine = FindGameFeaturePluginStateMachine(PluginURL))
 	{
 		//Define a lambda that will kick off the actual Terminate
-		const FGameFeaturePluginUpdateURLComplete StartTerminateLambda = FGameFeaturePluginUpdateURLComplete::CreateLambda([this, PluginURL, CompleteDelegate](const UE::GameFeatures::FResult& Result)
+		const FGameFeaturePluginUpdateURLComplete StartTerminateLambda = FGameFeaturePluginUpdateURLComplete::CreateWeakLambda(this, 
+			[this, PluginURL, CompleteDelegate](const UE::GameFeatures::FResult& Result)
 			{
 				UGameFeaturePluginStateMachine* StateMachine = FindGameFeaturePluginStateMachine(PluginURL);
 				check(StateMachine);
