Add logging for when GFP deactivation is paused and resumed so that its possibly to identify misbehaving pausers.

diff --git a/Engine/Plugins/Experimental/GameFeatures/Source/GameFeatures/Private/GameFeaturePluginStateMachine.cpp b/Engine/Plugins/Experimental/GameFeatures/Source/GameFeatures/Private/GameFeaturePluginStateMachine.cpp
--- a/Engine/Plugins/Experimental/GameFeatures/Source/GameFeatures/Private/GameFeaturePluginStateMachine.cpp
+++ b/Engine/Plugins/Experimental/GameFeatures/Source/GameFeatures/Private/GameFeaturePluginStateMachine.cpp
@@ -2317,12 +2317,14 @@ struct FGameFeaturePluginState_Deactivating : public FGameFeaturePluginState
 		bRequestedGC = false;
 	}
 
-	void OnPauserCompleted()
+	void OnPauserCompleted(FStringView InPauserTag)
 	{
 		check(IsInGameThread());
 		ensure(NumExpectedPausers != INDEX_NONE);
 		++NumObservedPausers;
 
+		UE_LOG(LogGameFeatures, Display, TEXT("Deactivation of %s resumed by %.*s"), *StateProperties.PluginName, InPauserTag.Len(), InPauserTag.GetData());
+
 		if (NumObservedPausers == NumExpectedPausers)
 		{
 			UpdateStateMachineImmediate();
@@ -2346,7 +2348,7 @@ struct FGameFeaturePluginState_Deactivating : public FGameFeaturePluginState
 			NumObservedPausers = 0;
 
 			// Deactivate
-			FGameFeatureDeactivatingContext Context(FSimpleDelegate::CreateRaw(this, &FGameFeaturePluginState_Deactivating::OnPauserCompleted));
+			FGameFeatureDeactivatingContext Context(StateProperties.PluginName, [this](FStringView InPauserTag) { OnPauserCompleted(InPauserTag); });
 			UGameFeaturesSubsystem::Get().OnGameFeatureDeactivating(StateProperties.GameFeatureData, StateProperties.PluginName, Context, StateProperties.PluginIdentifier.GetFullPluginURL());
 			NumExpectedPausers = Context.NumPausers;
 
diff --git a/Engine/Plugins/Experimental/GameFeatures/Source/GameFeatures/Private/GameFeaturesSubsystem.cpp b/Engine/Plugins/Experimental/GameFeatures/Source/GameFeatures/Private/GameFeaturesSubsystem.cpp
--- a/Engine/Plugins/Experimental/GameFeatures/Source/GameFeatures/Private/GameFeaturesSubsystem.cpp
+++ b/Engine/Plugins/Experimental/GameFeatures/Source/GameFeatures/Private/GameFeaturesSubsystem.cpp
@@ -235,6 +235,16 @@ bool FGameFeatureStateChangeContext::ShouldApplyUsingOtherContext(const FGameFea
 	return false;
 }
 
+FSimpleDelegate FGameFeatureDeactivatingContext::PauseDeactivationUntilComplete(FString InPauserTag)
+{
+	UE_LOG(LogGameFeatures, Display, TEXT("Deactivation of %.*s paused by %s"), PluginName.Len(), PluginName.GetData(), *InPauserTag);
+
+	++NumPausers;
+	return FSimpleDelegate::CreateLambda(
+		[CompletionCallback=CompletionCallback, PauserTag=MoveTemp(InPauserTag)]() { CompletionCallback(PauserTag); }
+	);
+}
+
 void UGameFeaturesSubsystem::Initialize(FSubsystemCollectionBase& Collection)
 {
 	UE_LOG(LogGameFeatures, Log, TEXT("Initializing game features subsystem"));
diff --git a/Engine/Plugins/Experimental/GameFeatures/Source/GameFeatures/Public/GameFeaturesSubsystem.h b/Engine/Plugins/Experimental/GameFeatures/Source/GameFeatures/Public/GameFeaturesSubsystem.h
--- a/Engine/Plugins/Experimental/GameFeatures/Source/GameFeatures/Public/GameFeaturesSubsystem.h
+++ b/Engine/Plugins/Experimental/GameFeatures/Source/GameFeatures/Public/GameFeaturesSubsystem.h
@@ -103,20 +103,18 @@ struct FGameFeatureDeactivatingContext : public FGameFeatureStateChangeContext
 {
 public:
 	// Call this if your observer has an asynchronous action to complete as part of shutdown, and invoke the returned delegate when you are done (on the game thread!)
-	FSimpleDelegate PauseDeactivationUntilComplete()
-	{
-		++NumPausers;
-		return CompletionDelegate;
-	}
+	GAMEFEATURES_API FSimpleDelegate PauseDeactivationUntilComplete(FString InPauserTag);
 
-	FGameFeatureDeactivatingContext(FSimpleDelegate&& InCompletionDelegate)
-		: CompletionDelegate(MoveTemp(InCompletionDelegate))
+	FGameFeatureDeactivatingContext(FStringView InPluginName, TFunction<void(FStringView InPauserTag)>&& InCompletionCallback)
+		: PluginName(InPluginName)
+		, CompletionCallback(MoveTemp(InCompletionCallback))
 	{
 	}
 
 	int32 GetNumPausers() const { return NumPausers; }
 private:
-	FSimpleDelegate CompletionDelegate;
+	FStringView PluginName;
+	TFunction<void(FStringView InPauserTag)> CompletionCallback;
 	int32 NumPausers = 0;
 
 	friend struct FGameFeaturePluginState_Deactivating;
