Add deperection for PauseDeactivationUntilComplete

diff --git a/Engine/Plugins/Experimental/GameFeatures/Source/GameFeatures/Public/GameFeaturesSubsystem.h b/Engine/Plugins/Experimental/GameFeatures/Source/GameFeatures/Public/GameFeaturesSubsystem.h
--- a/Engine/Plugins/Experimental/GameFeatures/Source/GameFeatures/Public/GameFeaturesSubsystem.h
+++ b/Engine/Plugins/Experimental/GameFeatures/Source/GameFeatures/Public/GameFeaturesSubsystem.h
@@ -102,6 +102,12 @@ private:
 struct FGameFeatureDeactivatingContext : public FGameFeatureStateChangeContext
 {
 public:
+	UE_DEPRECATED(5.2, "Use tagged version instead")
+	FSimpleDelegate PauseDeactivationUntilComplete()
+	{
+		PauseDeactivationUntilComplete(TEXT("Unknown(Deprecated)"));
+	}
+
 	// Call this if your observer has an asynchronous action to complete as part of shutdown, and invoke the returned delegate when you are done (on the game thread!)
 	GAMEFEATURES_API FSimpleDelegate PauseDeactivationUntilComplete(FString InPauserTag);
 
