Deperecate FGameFeatureDeactivatingContext constructor as well.

diff --git a/Engine/Plugins/Experimental/GameFeatures/Source/GameFeatures/Public/GameFeaturesSubsystem.h b/Engine/Plugins/Experimental/GameFeatures/Source/GameFeatures/Public/GameFeaturesSubsystem.h
--- a/Engine/Plugins/Experimental/GameFeatures/Source/GameFeatures/Public/GameFeaturesSubsystem.h
+++ b/Engine/Plugins/Experimental/GameFeatures/Source/GameFeatures/Public/GameFeaturesSubsystem.h
@@ -111,6 +111,12 @@ public:
 	// Call this if your observer has an asynchronous action to complete as part of shutdown, and invoke the returned delegate when you are done (on the game thread!)
 	GAMEFEATURES_API FSimpleDelegate PauseDeactivationUntilComplete(FString InPauserTag);
 
+	UE_DEPRECATED(5.2, "Use tagged version instead")
+	FGameFeatureDeactivatingContext(FSimpleDelegate&& InCompletionDelegate)
+		: CompletionCallback([CompletionDelegate = MoveTemp(InCompletionDelegate)](FStringView) { CompletionDelegate.ExecuteIfBound(); })
+	{
+	}
+
 	FGameFeatureDeactivatingContext(FStringView InPluginName, TFunction<void(FStringView InPauserTag)>&& InCompletionCallback)
 		: PluginName(InPluginName)
 		, CompletionCallback(MoveTemp(InCompletionCallback))
