And a follow up change for this one too

Deperecate FGameFeatureDeactivatingContext constructor as well.

diff --git a/Engine/Plugins/Experimental/GameFeatures/Source/GameFeatures/Public/GameFeaturesSubsystem.h b/Engine/Plugins/Experimental/GameFeatures/Source/GameFeatures/Public/GameFeaturesSubsystem.h
--- a/Engine/Plugins/Experimental/GameFeatures/Source/GameFeatures/Public/GameFeaturesSubsystem.h
+++ b/Engine/Plugins/Experimental/GameFeatures/Source/GameFeatures/Public/GameFeaturesSubsystem.h
@@ -113,7 +113,8 @@ public:
 
 	UE_DEPRECATED(5.2, "Use tagged version instead")
 	FGameFeatureDeactivatingContext(FSimpleDelegate&& InCompletionDelegate)
-		: CompletionCallback([CompletionDelegate = MoveTemp(InCompletionDelegate)](FStringView) { CompletionDelegate.ExecuteIfBound(); })
+		: PluginName(TEXTVIEW("Unknown(Deprecated)"))
+		, CompletionCallback([CompletionDelegate = MoveTemp(InCompletionDelegate)](FStringView) { CompletionDelegate.ExecuteIfBound(); })
 	{
 	}
 
