Fix RemoveGameFeatureFromAssetManager to correctly handle absolue paths and not always add the plugin name.
This is a followup to CL 18694050, the remove function was added in a different branch before the path fixup function existed there

diff --git a/Engine/Plugins/Experimental/GameFeatures/Source/GameFeatures/Private/GameFeaturesSubsystem.cpp b/Engine/Plugins/Experimental/GameFeatures/Source/GameFeatures/Private/GameFeaturesSubsystem.cpp
--- a/Engine/Plugins/Experimental/GameFeatures/Source/GameFeatures/Private/GameFeaturesSubsystem.cpp
+++ b/Engine/Plugins/Experimental/GameFeatures/Source/GameFeatures/Private/GameFeaturesSubsystem.cpp
@@ -535,6 +535,7 @@ void UGameFeaturesSubsystem::AddGameFeatureToAssetManager(const UGameFeatureData
 void UGameFeaturesSubsystem::RemoveGameFeatureFromAssetManager(const UGameFeatureData* GameFeatureToRemove, const FString& PluginName, const TArray<FName>& AddedPrimaryAssetTypes)
 {
 	check(GameFeatureToRemove);
+	FString PluginRootPath = TEXT("/") + PluginName + TEXT("/");
 	UAssetManager& LocalAssetManager = UAssetManager::Get();
 
 	for (FPrimaryAssetTypeInfo TypeInfo : GameFeatureToRemove->GetPrimaryAssetTypesToScan())
@@ -547,7 +548,7 @@ void UGameFeaturesSubsystem::RemoveGameFeatureFromAssetManager(const UGameFeatur
 
 		for (FDirectoryPath& Path : TypeInfo.Directories)
 		{
-			Path.Path = TEXT("/") + PluginName + TEXT("/") + Path.Path;
+			FixPluginPackagePath(Path.Path, PluginRootPath, false);
 		}
 
 		// This function also fills out runtime data on the copy
