Fix comment typo.

diff --git a/Engine/Plugins/Experimental/GameFeatures/Source/GameFeatures/Public/GameFeatureTypes.h b/Engine/Plugins/Experimental/GameFeatures/Source/GameFeatures/Public/GameFeatureTypes.h
--- a/Engine/Plugins/Experimental/GameFeatures/Source/GameFeatures/Public/GameFeatureTypes.h
+++ b/Engine/Plugins/Experimental/GameFeatures/Source/GameFeatures/Public/GameFeatureTypes.h
@@ -29,7 +29,7 @@ class FString;
 	XSTATE(Unregistering,				NSLOCTEXT("GameFeatures", "UnregisteringStateDisplayName", "Unregistering"))								/* Transition state Registered -> Installed. Cleaning up data gathered in Registering. */ \
 	XSTATE(Registering,					NSLOCTEXT("GameFeatures", "RegisteringStateDisplayName", "Registering"))									/* Transition state Installed -> Registered. Discovering assets in the plugin, but not loading them, except a few for discovery reasons. */ \
 	XSTATE(Registered,					NSLOCTEXT("GameFeatures", "RegisteredStateDisplayName", "Registered"))										/* The assets in the plugin are known, but have not yet been loaded, except a few for discovery reasons. */ \
-	XSTATE(Unloading,					NSLOCTEXT("GameFeatures", "UnloadingStateDisplayName", "Unloading"))										/* Transition state Loaded -> Registered. In the process of removing code/contnet from memory. */ \
+	XSTATE(Unloading,					NSLOCTEXT("GameFeatures", "UnloadingStateDisplayName", "Unloading"))										/* Transition state Loaded -> Registered. In the process of removing code/content from memory. */ \
 	XSTATE(Loading,						NSLOCTEXT("GameFeatures", "LoadingStateDisplayName", "Loading"))											/* Transition state Registered -> Loaded. In the process of loading code/content into memory. */ \
 	XSTATE(Loaded,						NSLOCTEXT("GameFeatures", "LoadedStateDisplayName", "Loaded"))												/* The plugin is loaded into memory, but not registered with game systems and active. */ \
 	XSTATE(Deactivating,				NSLOCTEXT("GameFeatures", "DeactivatingStateDisplayName", "Deactivating"))									/* Transition state Active -> Loaded. Currently unregistering with game systems. */ \
