GFP
- Fix GFP state machine keys getting stomped when FGameFeaturePluginState::TryUpdateURLData is called.
- Fix missing '.' in error codes

diff --git a/Engine/Plugins/Experimental/GameFeatures/Source/GameFeatures/Private/GameFeaturePluginStateMachine.cpp b/Engine/Plugins/Experimental/GameFeatures/Source/GameFeatures/Private/GameFeaturePluginStateMachine.cpp
--- a/Engine/Plugins/Experimental/GameFeatures/Source/GameFeatures/Private/GameFeaturePluginStateMachine.cpp
+++ b/Engine/Plugins/Experimental/GameFeatures/Source/GameFeatures/Private/GameFeaturePluginStateMachine.cpp
@@ -482,7 +482,7 @@ UE::GameFeatures::FResult FGameFeaturePluginState::GetErrorResult(const FString&
 {
 	const FString StateName = LexToString(UGameFeaturesSubsystem::Get().GetPluginState(StateProperties.PluginIdentifier));
 	const FString ErrorCodeEnding = ErrorNamespaceAddition.IsEmpty() ? ErrorCode : ErrorNamespaceAddition + ErrorCode;
-	const FString CompleteErrorCode = (UE::GameFeatures::StateMachineErrorNamespace + StateName + ErrorCodeEnding);
+	const FString CompleteErrorCode = FString::Printf(TEXT("%s%s.%s"), *UE::GameFeatures::StateMachineErrorNamespace, *StateName, *ErrorCodeEnding);
 	return UE::GameFeatures::FResult(MakeError(CompleteErrorCode), OptionalErrorText);
 }
 
diff --git a/Engine/Plugins/Experimental/GameFeatures/Source/GameFeatures/Private/GameFeaturesSubsystem.cpp b/Engine/Plugins/Experimental/GameFeatures/Source/GameFeatures/Private/GameFeaturesSubsystem.cpp
--- a/Engine/Plugins/Experimental/GameFeatures/Source/GameFeatures/Private/GameFeaturesSubsystem.cpp
+++ b/Engine/Plugins/Experimental/GameFeatures/Source/GameFeatures/Private/GameFeaturesSubsystem.cpp
@@ -161,6 +161,21 @@ FGameFeaturePluginIdentifier::FGameFeaturePluginIdentifier(FString PluginURL)
 	FromPluginURL(MoveTemp(PluginURL));
 }
 
+FGameFeaturePluginIdentifier::FGameFeaturePluginIdentifier(FGameFeaturePluginIdentifier&& Other)
+	: FGameFeaturePluginIdentifier(MoveTemp(Other.PluginURL))
+{
+	Other.IdentifyingURLSubset.Reset();
+	Other.PluginProtocol = EGameFeaturePluginProtocol::Unknown;
+}
+
+FGameFeaturePluginIdentifier& FGameFeaturePluginIdentifier::operator=(FGameFeaturePluginIdentifier&& Other)
+{
+	FromPluginURL(MoveTemp(Other.PluginURL));
+	Other.IdentifyingURLSubset.Reset();
+	Other.PluginProtocol = EGameFeaturePluginProtocol::Unknown;
+	return *this;
+}
+
 void FGameFeaturePluginIdentifier::FromPluginURL(FString PluginURLIn)
 {
 	//Make sure we have no stale data
diff --git a/Engine/Plugins/Experimental/GameFeatures/Source/GameFeatures/Public/GameFeaturesSubsystem.h b/Engine/Plugins/Experimental/GameFeatures/Source/GameFeatures/Public/GameFeaturesSubsystem.h
--- a/Engine/Plugins/Experimental/GameFeatures/Source/GameFeatures/Public/GameFeaturesSubsystem.h
+++ b/Engine/Plugins/Experimental/GameFeatures/Source/GameFeatures/Public/GameFeaturesSubsystem.h
@@ -231,6 +231,20 @@ struct FGameFeaturePluginIdentifier
 	FGameFeaturePluginIdentifier() = default;
 	FGameFeaturePluginIdentifier(FString PluginURL);
 
+	FGameFeaturePluginIdentifier(const FGameFeaturePluginIdentifier& Other)
+		: FGameFeaturePluginIdentifier(Other.PluginURL)
+	{}
+
+	FGameFeaturePluginIdentifier(FGameFeaturePluginIdentifier&& Other);
+
+	FGameFeaturePluginIdentifier& operator=(const FGameFeaturePluginIdentifier& Other)
+	{
+		FromPluginURL(Other.PluginURL);
+		return *this;
+	}
+
+	FGameFeaturePluginIdentifier& operator=(FGameFeaturePluginIdentifier&& Other);
+
 	/** Used to determine if 2 FGameFeaturePluginIdentifiers are referencing the same GameFeaturePlugin.
 		Only matching on Identifying information instead of all the optional bundle information */
 	bool operator==(const FGameFeaturePluginIdentifier& Other) const;
@@ -256,14 +270,14 @@ struct FGameFeaturePluginIdentifier
 	}
 
 private:
-	/** The protocol used in the URL for this GameFeaturePlugin URL */
-	EGameFeaturePluginProtocol PluginProtocol;
+	/** Full PluginURL used to originally construct this identifier */
+	FString PluginURL;
 
 	/** The part of the URL that can be used to uniquely identify this plugin without any transient data */
 	FStringView IdentifyingURLSubset;
 
-	/** Full PluginURL used to originally construct this identifier */
-	FString PluginURL;
+	/** The protocol used in the URL for this GameFeaturePlugin URL */
+	EGameFeaturePluginProtocol PluginProtocol;
 
 	//Friend class so that it can access parsed URL data from under the hood
 	friend struct FGameFeaturePluginStateMachineProperties;
