Follow up to 24721344 to simplify GFP state machine tracking
- Use string keys for GFP state machines instead of FGameFeaturePluginIdentifier
- Keys are still the identifying part of the URL

diff --git a/Engine/Plugins/Experimental/GameFeatures/Source/GameFeatures/Private/GameFeaturesSubsystem.cpp b/Engine/Plugins/Experimental/GameFeatures/Source/GameFeatures/Private/GameFeaturesSubsystem.cpp
--- a/Engine/Plugins/Experimental/GameFeatures/Source/GameFeatures/Private/GameFeaturesSubsystem.cpp
+++ b/Engine/Plugins/Experimental/GameFeatures/Source/GameFeatures/Private/GameFeaturesSubsystem.cpp
@@ -1530,9 +1530,10 @@ UGameFeaturePluginStateMachine* UGameFeaturesSubsystem::FindGameFeaturePluginSta
 	return FindGameFeaturePluginStateMachine(FindPluginIdentifier);
 }
 
-UGameFeaturePluginStateMachine* UGameFeaturesSubsystem::FindGameFeaturePluginStateMachine(FGameFeaturePluginIdentifier PluginIdentifier) const
+UGameFeaturePluginStateMachine* UGameFeaturesSubsystem::FindGameFeaturePluginStateMachine(const FGameFeaturePluginIdentifier& PluginIdentifier) const
 {
-	TObjectPtr<UGameFeaturePluginStateMachine> const* ExistingStateMachine = GameFeaturePluginStateMachines.Find(PluginIdentifier);
+	TObjectPtr<UGameFeaturePluginStateMachine> const* ExistingStateMachine = 
+		GameFeaturePluginStateMachines.FindByHash(GetTypeHash(PluginIdentifier.GetIdentifyingString()), PluginIdentifier.GetIdentifyingString());
 	if (ExistingStateMachine)
 	{
 		UE_LOG(LogGameFeatures, VeryVerbose, TEXT("FOUND GameFeaturePlugin using PluginIdentifier:%.*s for PluginURL:%s"), PluginIdentifier.GetIdentifyingString().Len(), PluginIdentifier.GetIdentifyingString().GetData(), *PluginIdentifier.GetFullPluginURL());
@@ -1556,7 +1557,7 @@ UGameFeaturePluginStateMachine* UGameFeaturesSubsystem::FindOrCreateGameFeatureP
 	UE_LOG(LogGameFeatures, Display, TEXT("Creating GameFeaturePlugin StateMachine using Identifier:%.*s from PluginURL:%s"), PluginIdentifier.GetIdentifyingString().Len(), PluginIdentifier.GetIdentifyingString().GetData(), *PluginURL);
 
 	UGameFeaturePluginStateMachine* NewStateMachine = NewObject<UGameFeaturePluginStateMachine>(this);
-	GameFeaturePluginStateMachines.Add(PluginIdentifier, NewStateMachine);
+	GameFeaturePluginStateMachines.Add(FString(PluginIdentifier.GetIdentifyingString()), NewStateMachine);
 	NewStateMachine->InitStateMachine(MoveTemp(PluginIdentifier));
 
 	return NewStateMachine;
@@ -1655,8 +1656,11 @@ void UGameFeaturesSubsystem::BeginTermination(UGameFeaturePluginStateMachine* Ma
 {
 	check(IsValid(Machine));
 	check(Machine->GetCurrentState() == EGameFeaturePluginState::Terminal);
-	UE_LOG(LogGameFeatures, Verbose, TEXT("BeginTermination of GameFeaturePlugin. Identifier:%.*s URL:%s"), Machine->GetPluginIdentifier().GetIdentifyingString().Len(), Machine->GetPluginIdentifier().GetIdentifyingString().GetData(), *(Machine->GetPluginURL()));
-	GameFeaturePluginStateMachines.Remove(Machine->GetPluginIdentifier());
+
+	FStringView Identifer = Machine->GetPluginIdentifier().GetIdentifyingString();
+
+	UE_LOG(LogGameFeatures, Verbose, TEXT("BeginTermination of GameFeaturePlugin. Identifier:%.*s URL:%s"), Identifer.Len(), Identifer.GetData(), *(Machine->GetPluginURL()));
+	GameFeaturePluginStateMachines.RemoveByHash(GetTypeHash(Identifer), Identifer);
 	TerminalGameFeaturePluginStateMachines.Add(Machine);
 }
 
@@ -1836,7 +1840,7 @@ TSet<FString> UGameFeaturesSubsystem::GetActivePluginNames() const
 {
 	TSet<FString> ActivePluginNames;
 
-	for (const TPair<FGameFeaturePluginIdentifier, TObjectPtr<UGameFeaturePluginStateMachine>>& Pair : GameFeaturePluginStateMachines)
+	for (const TPair<FString, TObjectPtr<UGameFeaturePluginStateMachine>>& Pair : GameFeaturePluginStateMachines)
 	{
 		UGameFeaturePluginStateMachine* StateMachine = Pair.Value;
 		if (StateMachine->GetCurrentState() == EGameFeaturePluginState::Active &&
diff --git a/Engine/Plugins/Experimental/GameFeatures/Source/GameFeatures/Public/GameFeaturesSubsystem.h b/Engine/Plugins/Experimental/GameFeatures/Source/GameFeatures/Public/GameFeaturesSubsystem.h
--- a/Engine/Plugins/Experimental/GameFeatures/Source/GameFeatures/Public/GameFeaturesSubsystem.h
+++ b/Engine/Plugins/Experimental/GameFeatures/Source/GameFeatures/Public/GameFeaturesSubsystem.h
@@ -229,7 +229,7 @@ struct FGameFeaturePluginIdentifier
 	GENERATED_BODY()
 
 	FGameFeaturePluginIdentifier() = default;
-	FGameFeaturePluginIdentifier(FString PluginURL);
+	explicit FGameFeaturePluginIdentifier(FString PluginURL);
 
 	FGameFeaturePluginIdentifier(const FGameFeaturePluginIdentifier& Other)
 		: FGameFeaturePluginIdentifier(Other.PluginURL)
@@ -256,10 +256,10 @@ struct FGameFeaturePluginIdentifier
 		To match exactly all information in the PluginURL has to match and not just the IdentifyingURLSubset */
 	bool ExactMatchesURL(const FString& PluginURL) const;
 
-	const EGameFeaturePluginProtocol GetPluginProtocol() const { return PluginProtocol; }
+	EGameFeaturePluginProtocol GetPluginProtocol() const { return PluginProtocol; }
 
 	/** Returns the Identifying information used for this Plugin. It is a subset of the URL used to create it.*/
-	const FStringView GetIdentifyingString() const { return IdentifyingURLSubset; }
+	FStringView GetIdentifyingString() const { return IdentifyingURLSubset; }
 
 	/** Get the Full PluginURL used to originally construct this identifier */
 	const FString& GetFullPluginURL() const { return PluginURL; }
@@ -587,7 +587,7 @@ private:
 	UGameFeaturePluginStateMachine* FindGameFeaturePluginStateMachine(const FString& PluginURL) const;
 
 	/** Gets the state machine associated with the specified PluginIdentifier */
-	UGameFeaturePluginStateMachine* FindGameFeaturePluginStateMachine(FGameFeaturePluginIdentifier PluginIdentifier) const;
+	UGameFeaturePluginStateMachine* FindGameFeaturePluginStateMachine(const FGameFeaturePluginIdentifier& PluginIdentifier) const;
 
 	/** Gets the state machine associated with the specified URL, creates it if it doesnt exist */
 	UGameFeaturePluginStateMachine* FindOrCreateGameFeaturePluginStateMachine(const FString& PluginURL);
@@ -633,7 +633,7 @@ private:
 private:
 	/** The list of all game feature plugin state machine objects */
 	UPROPERTY(Transient)
-	TMap<FGameFeaturePluginIdentifier, TObjectPtr<UGameFeaturePluginStateMachine>> GameFeaturePluginStateMachines;
+	TMap<FString, TObjectPtr<UGameFeaturePluginStateMachine>> GameFeaturePluginStateMachines;
 
 	/** Game feature plugin state machine objects that are being terminated. Used to prevent GC until termination is complete. */
 	UPROPERTY(Transient)
