Batch assets unloading in editor to speed up umounting several game feature plugins simultaneously
Closing a project with 370 unloaded plugins went from 21.5 sec to 9.5 sec (in dev build)

diff --git a/Engine/Plugins/Experimental/GameFeatures/Source/GameFeatures/Private/GameFeaturePluginStateMachine.cpp b/Engine/Plugins/Experimental/GameFeatures/Source/GameFeatures/Private/GameFeaturePluginStateMachine.cpp
--- a/Engine/Plugins/Experimental/GameFeatures/Source/GameFeatures/Private/GameFeaturePluginStateMachine.cpp
+++ b/Engine/Plugins/Experimental/GameFeatures/Source/GameFeatures/Private/GameFeaturePluginStateMachine.cpp
@@ -384,6 +384,49 @@ namespace UE::GameFeatures
 			}
 		}
 	};
+
+#if WITH_EDITOR
+	TMap<FString, FGameFeaturePluginRequestUpdateStateMachine> PluginsToUnloadAssets;
+	FTSTicker::FDelegateHandle UnloadPluginAssetsHandle;
+
+	bool TickUnloadPluginAssets(float /*DeltaTime*/)
+	{
+		UnloadPluginAssetsHandle.Reset();
+	
+		TArray<FString> PluginNames;
+		TArray<FGameFeaturePluginRequestUpdateStateMachine> UpdateStateMachineDelegates;
+		{
+			PluginNames.Reserve(PluginsToUnloadAssets.Num());
+			UpdateStateMachineDelegates.Reserve(PluginsToUnloadAssets.Num());
+			for (TPair<FString, FGameFeaturePluginRequestUpdateStateMachine>& PluginsToUnloadAsset : PluginsToUnloadAssets)
+			{
+				PluginNames.Add(PluginsToUnloadAsset.Key);
+				UpdateStateMachineDelegates.Add(MoveTemp(PluginsToUnloadAsset.Value));
+			}
+			PluginsToUnloadAssets.Empty();
+		}
+
+		verify(FPluginUtils::UnloadPluginsAssets(PluginNames));
+
+		for (const FGameFeaturePluginRequestUpdateStateMachine& UpdateStateMachineDelegate : UpdateStateMachineDelegates)
+		{
+			UpdateStateMachineDelegate.ExecuteIfBound();
+		}
+
+		return false;
+	}
+
+	void ScheduleUnloadPluginAssets(const FString& PluginName, const FGameFeaturePluginRequestUpdateStateMachine& UpdateStateMachineDelegate)
+	{
+		check(IsInGameThread());
+		ensure(!PluginsToUnloadAssets.Contains(PluginName));
+		PluginsToUnloadAssets.Add(PluginName, UpdateStateMachineDelegate);
+		if (!UnloadPluginAssetsHandle.IsValid())
+		{
+			UnloadPluginAssetsHandle = FTSTicker::GetCoreTicker().AddTicker(FTickerDelegate::CreateStatic(&TickUnloadPluginAssets));
+		}
+	}
+#endif //if WITH_EDITOR
 }
 
 //Enum describing the options for what data can be inside the InstallBundle protocol URL's metadata
@@ -2038,10 +2081,16 @@ struct FGameFeaturePluginState_Unregistering : public FGameFeaturePluginState
 	FGameFeaturePluginState_Unregistering(FGameFeaturePluginStateMachineProperties& InStateProperties) : FGameFeaturePluginState(InStateProperties) {}
 
 	bool bRequestedGC = false;
+#if WITH_EDITOR
+	bool bRequestedUnloadPluginAssets = false;
+#endif //if WITH_EDITOR
 
 	virtual void BeginState()
 	{
 		bRequestedGC = false;
+#if WITH_EDITOR
+		bRequestedUnloadPluginAssets = false;
+#endif //if WITH_EDITOR
 	}
 
 	virtual void UpdateState(FGameFeaturePluginStateStatus& StateStatus) override
@@ -2054,37 +2103,45 @@ struct FGameFeaturePluginState_Unregistering : public FGameFeaturePluginState
 			return;
 		}
 
+#if WITH_EDITOR
+		if (bRequestedUnloadPluginAssets)
+		{
+			bRequestedGC = true;
+			GarbageCollectAndUpdateStateMachineDeferred(/*bWaitForCleanupQueue=*/true);
+			return;
+		}
+#endif //if WITH_EDITOR
+
 		if (StateProperties.GameFeatureData)
 		{
 			UGameFeaturesSubsystem::Get().OnGameFeatureUnregistering(StateProperties.GameFeatureData, StateProperties.PluginName, StateProperties.PluginIdentifier.GetFullPluginURL());
 
-			UGameFeaturesSubsystem::Get().RemoveGameFeatureFromAssetManager(StateProperties.GameFeatureData, StateProperties.PluginName, StateProperties.AddedPrimaryAssetTypes);
+			UGameFeaturesSubsystem::RemoveGameFeatureFromAssetManager(StateProperties.GameFeatureData, StateProperties.PluginName, StateProperties.AddedPrimaryAssetTypes);
 			StateProperties.AddedPrimaryAssetTypes.Empty();
 
-			UGameFeaturesSubsystem::Get().UnloadGameFeatureData(StateProperties.GameFeatureData);
+			UGameFeaturesSubsystem::UnloadGameFeatureData(StateProperties.GameFeatureData);
 		}
-
 		StateProperties.GameFeatureData = nullptr;
 
+		// Try to remove the gameplay tags, this might be ignored depending on project settings
+		const FString PluginFolder = FPaths::GetPath(StateProperties.PluginInstalledFilename);
+		UGameplayTagsManager::Get().RemoveTagIniSearchPath(PluginFolder / TEXT("Config") / TEXT("Tags"));
+
 #if WITH_EDITOR
 		// This will properly unload any plugin asset that could be opened in the editor
 		// and ensure standalone packages get unloaded as well
-		verify(FPluginUtils::UnloadPluginAssets(StateProperties.PluginName));
-#endif //if WITH_EDITOR
-
-#if !WITH_EDITOR // never mark anything as garbage in the editor, there be dwagons here. 
+		bRequestedUnloadPluginAssets = true;
+		UE::GameFeatures::ScheduleUnloadPluginAssets(StateProperties.PluginName, StateProperties.OnRequestUpdateStateMachine);
+#else
+		// never mark anything as garbage in the editor, there be dwagons here. 
 		if (UE::GameFeatures::MarkPluginAsGarbageOnUnregister && !UE::GameFeatures::ShouldSkipVerify(StateProperties.PluginName))
 		{
 			MarkPluginAsGarbage(true);
 		}
-#endif
-
-		// Try to remove the gameplay tags, this might be ignored depending on project settings
-		const FString PluginFolder = FPaths::GetPath(StateProperties.PluginInstalledFilename);
-		UGameplayTagsManager::Get().RemoveTagIniSearchPath(PluginFolder / TEXT("Config") / TEXT("Tags"));
 
 		bRequestedGC = true;
-		GarbageCollectAndUpdateStateMachineDeferred(true);
+		GarbageCollectAndUpdateStateMachineDeferred(/*bWaitForCleanupQueue=*/true);
+#endif
 	}
 };
 
