Fixing error where multiple bundles are not correctly serialized in FInstallBundlePluginProtocolMetaData::ToString. This leads to problems when constructing InstallBundle GameFeaturePlugin URLs from FInstallBundlePluginProtocolMetadata.
#tests Windows Dev Client

diff --git a/Engine/Plugins/Experimental/GameFeatures/Source/GameFeatures/Private/GameFeaturePluginStateMachine.cpp b/Engine/Plugins/Experimental/GameFeatures/Source/GameFeatures/Private/GameFeaturePluginStateMachine.cpp
--- a/Engine/Plugins/Experimental/GameFeatures/Source/GameFeatures/Private/GameFeaturePluginStateMachine.cpp
+++ b/Engine/Plugins/Experimental/GameFeatures/Source/GameFeatures/Private/GameFeaturePluginStateMachine.cpp
@@ -3007,9 +3007,11 @@ FString FInstallBundlePluginProtocolMetaData::ToString() const
 
 	//Always encode InstallBundles
 	ReturnedString = LexToString(EGameFeatureInstallBundleProtocolOptions::Bundles) + UE::GameFeatures::PluginURLStructureInfo::OptionAssignOperator;
-	for (const FName& BundleName : InstallBundles)
+	
+	if (InstallBundles.Num() > 0)
 	{
-		ReturnedString.Append(BundleName.ToString());
+		const FString BundlesList = FString::JoinBy(InstallBundles, UE::GameFeatures::PluginURLStructureInfo::OptionListSeperator, [](const FName& BundleName){ return FNameBuilder(BundleName); });
+		ReturnedString.Append(BundlesList);
 	}
 
 	//Always encode version
diff --git a/Engine/Plugins/Experimental/GameFeatures/Source/GameFeatures/Private/GameFeaturesSubsystem.cpp b/Engine/Plugins/Experimental/GameFeatures/Source/GameFeatures/Private/GameFeaturesSubsystem.cpp
--- a/Engine/Plugins/Experimental/GameFeatures/Source/GameFeatures/Private/GameFeaturesSubsystem.cpp
+++ b/Engine/Plugins/Experimental/GameFeatures/Source/GameFeatures/Private/GameFeaturesSubsystem.cpp
@@ -41,6 +41,7 @@ namespace UE::GameFeatures
 	{
 		const TCHAR* OptionAssignOperator = TEXT("=");
 		const TCHAR* OptionSeperator = TEXT("?");
+		const TCHAR* OptionListSeperator = TEXT(",");
 	}
 
 	namespace CommonErrorCodes
diff --git a/Engine/Plugins/Experimental/GameFeatures/Source/GameFeatures/Public/GameFeaturesSubsystem.h b/Engine/Plugins/Experimental/GameFeatures/Source/GameFeatures/Public/GameFeaturesSubsystem.h
--- a/Engine/Plugins/Experimental/GameFeatures/Source/GameFeatures/Public/GameFeaturesSubsystem.h
+++ b/Engine/Plugins/Experimental/GameFeatures/Source/GameFeatures/Public/GameFeaturesSubsystem.h
@@ -27,8 +27,14 @@ namespace UE::GameFeatures
 {
 	namespace PluginURLStructureInfo
 	{
+		/** Character used to denote what value is being assigned to the option before it */
 		extern const TCHAR* OptionAssignOperator;
+
+		/** Character used to separate options on the URL. Used between each assigned value and the next Option name. */
 		extern const TCHAR* OptionSeperator;
+
+		/** Character used to separate lists of values for a single option. Used between each entry in the list. */
+		extern const TCHAR* OptionListSeperator;
 	};
 
 	namespace CommonErrorCodes
