Add policy based GFP dependency resolution.
- Fix explicitly loaded plugins never getting removed from plugin details cache if they get removed from plugin manager.
- In editor, load plugin details from cached json instead of loading uplugin multiple times.

diff --git a/Engine/Plugins/Experimental/GameFeatures/Source/GameFeatures/Private/GameFeaturePluginStateMachine.cpp b/Engine/Plugins/Experimental/GameFeatures/Source/GameFeatures/Private/GameFeaturePluginStateMachine.cpp
--- a/Engine/Plugins/Experimental/GameFeatures/Source/GameFeatures/Private/GameFeaturePluginStateMachine.cpp
+++ b/Engine/Plugins/Experimental/GameFeatures/Source/GameFeatures/Private/GameFeaturePluginStateMachine.cpp
@@ -1645,6 +1645,8 @@ struct FGameFeaturePluginState_Unmounting : public FGameFeaturePluginState
 
 		if (StateProperties.bAddedPluginToManager)
 		{
+			UGameFeaturesSubsystem::Get().PruneCachedGameFeaturePluginDetails(StateProperties.PluginIdentifier.GetFullPluginURL(), StateProperties.PluginInstalledFilename);
+
 			verify(IPluginManager::Get().RemoveFromPluginsList(StateProperties.PluginInstalledFilename));
 			StateProperties.bAddedPluginToManager = false;
 		}
@@ -1937,7 +1939,7 @@ struct FGameFeaturePluginState_WaitingForDependencies : public FGameFeaturePlugi
 			UGameFeaturesSubsystem& GameFeaturesSubsystem = UGameFeaturesSubsystem::Get();
 
 			TArray<UGameFeaturePluginStateMachine*> Dependencies;
-			if (!GameFeaturesSubsystem.FindOrCreatePluginDependencyStateMachines(StateProperties.PluginInstalledFilename, Dependencies))
+			if (!GameFeaturesSubsystem.FindOrCreatePluginDependencyStateMachines(*StateProperties.PluginIdentifier.GetFullPluginURL(), StateProperties.PluginInstalledFilename, Dependencies))
 			{
 				// Failed to query dependencies
 				StateStatus.SetTransitionError(EGameFeaturePluginState::ErrorWaitingForDependencies, GetErrorResult(TEXT("Failed_Dependency_Query")));
diff --git a/Engine/Plugins/Experimental/GameFeatures/Source/GameFeatures/Private/GameFeaturesProjectPolicies.cpp b/Engine/Plugins/Experimental/GameFeatures/Source/GameFeatures/Private/GameFeaturesProjectPolicies.cpp
--- a/Engine/Plugins/Experimental/GameFeatures/Source/GameFeatures/Private/GameFeaturesProjectPolicies.cpp
+++ b/Engine/Plugins/Experimental/GameFeatures/Source/GameFeatures/Private/GameFeaturesProjectPolicies.cpp
@@ -2,6 +2,8 @@
 
 #include "GameFeaturesProjectPolicies.h"
 #include "GameFeaturesSubsystemSettings.h"
+#include "Interfaces/IPluginManager.h"
+#include "Misc/Paths.h"
 
 #include UE_INLINE_GENERATED_CPP_BY_NAME(GameFeaturesProjectPolicies)
 
@@ -47,6 +49,38 @@ bool UGameFeaturesProjectPolicies::WillPluginBeCooked(const FString& PluginFilen
 	return true;
 }
 
+TValueOrError<FString, FString> UGameFeaturesProjectPolicies::ReslovePluginDependency(const FString& PluginURL, const FString& DependencyName) const
+{
+	FString DependencyURL;
+	bool bResolvedDependency = false;
+
+	// Check if the dependency plugin exists yet (should be true for all built-in plugins)
+	if (TSharedPtr<IPlugin> DependencyPlugin = IPluginManager::Get().FindPlugin(DependencyName))
+	{
+		// Check if UGameFeaturesSubsystem is already aware of it
+		if (!UGameFeaturesSubsystem::Get().GetPluginURLByName(DependencyPlugin->GetName(), DependencyURL))
+		{
+			// It could still be a GFP, but state machine may not have been created for it yet
+			// Check if it is a built-in GFP
+			if (!DependencyPlugin->GetDescriptorFileName().IsEmpty() &&
+				GetDefault<UGameFeaturesSubsystemSettings>()->IsValidGameFeaturePlugin(FPaths::ConvertRelativePathToFull(DependencyPlugin->GetDescriptorFileName())) &&
+				FPaths::FileExists(DependencyPlugin->GetDescriptorFileName()))
+			{
+				DependencyURL = UGameFeaturesSubsystem::GetPluginURL_FileProtocol(DependencyPlugin->GetDescriptorFileName());
+			}
+		}
+
+		bResolvedDependency = true;
+	}
+
+	if (bResolvedDependency)
+	{
+		return MakeValue(MoveTemp(DependencyURL));
+	}
+
+	return MakeError(TEXT("NotFound"));
+}
+
 void UGameFeaturesProjectPolicies::ExplicitLoadGameFeaturePlugin(const FString& PluginURL, const FGameFeaturePluginLoadComplete& CompleteDelegate, const bool bActivateGameFeatures)
 {
 	if (bActivateGameFeatures)
diff --git a/Engine/Plugins/Experimental/GameFeatures/Source/GameFeatures/Private/GameFeaturesSubsystem.cpp b/Engine/Plugins/Experimental/GameFeatures/Source/GameFeatures/Private/GameFeaturesSubsystem.cpp
--- a/Engine/Plugins/Experimental/GameFeatures/Source/GameFeatures/Private/GameFeaturesSubsystem.cpp
+++ b/Engine/Plugins/Experimental/GameFeatures/Source/GameFeatures/Private/GameFeaturesSubsystem.cpp
@@ -12,6 +12,7 @@
 #include "Interfaces/IPluginManager.h"
 #include "Misc/PackageName.h"
 #include "Misc/Paths.h"
+#include "Misc/PathViews.h"
 #include "Misc/FileHelper.h"
 #include "Serialization/JsonSerializer.h"
 #include "Stats/StatsMisc.h"
@@ -1359,8 +1360,11 @@ EGameFeaturePluginState UGameFeaturesSubsystem::GetPluginState(FGameFeaturePlugi
 	}
 }
 
-bool UGameFeaturesSubsystem::GetGameFeaturePluginDetails(const TSharedRef<IPlugin>& Plugin, FString& OutPluginURL, struct FGameFeaturePluginDetails& OutPluginDetails) const
+bool UGameFeaturesSubsystem::GetGameFeaturePluginDetails(const TSharedRef<IPlugin>& Plugin, FString& OutPluginURL, FGameFeaturePluginDetails& OutPluginDetails) const
 {
+	// @TODO: this problematic because it assumes file protocol.
+	// Ideally this would work with any protocol, but for current uses cases the exact protocol doesn't seem to matter.
+
 	const FString& PluginDescriptorFilename = Plugin->GetDescriptorFileName();
 
 	// Make sure you are in a game feature plugins folder. All GameFeaturePlugins are rooted in a GameFeatures folder.
@@ -1378,14 +1382,14 @@ bool UGameFeaturesSubsystem::GetGameFeaturePluginDetails(const TSharedRef<IPlugi
 
 		if (bIsFileProtocol)
 		{
-			return GetGameFeaturePluginDetails(PluginDescriptorFilename, OutPluginDetails);
+			return GetGameFeaturePluginDetails(OutPluginURL, PluginDescriptorFilename, OutPluginDetails);
 		}
 	}
 
 	return false;
 }
 
-bool UGameFeaturesSubsystem::GetGameFeaturePluginDetails(const FString& PluginDescriptorFilename, FGameFeaturePluginDetails& OutPluginDetails) const
+bool UGameFeaturesSubsystem::GetGameFeaturePluginDetails(const FString& PluginURL, const FString& PluginDescriptorFilename, FGameFeaturePluginDetails& OutPluginDetails) const
 {
 	TRACE_CPUPROFILER_EVENT_SCOPE(GFP_GetPluginDetails);
 
@@ -1395,6 +1399,9 @@ bool UGameFeaturesSubsystem::GetGameFeaturePluginDetails(const FString& PluginDe
 	FDateTime FileTimeStamp;
 	if (UE::GameFeatures::GCachePluginDetails)
 	{
+		// Note: On file systems that don't support timestamps (current time is returned instead), 
+		// the pak file layer will end up caching the mount time, so this stamp will still be valid as long as the uplugin
+		// is in a pak and the pak is mounted.
 		FileTimeStamp = IFileManager::Get().GetTimeStamp(*PluginDescriptorFilename);
 		if (FCachedGameFeaturePluginDetails* ExistingDetails = CachedPluginDetailsByFilename.Find(PluginDescriptorFilename))
 		{
@@ -1410,26 +1417,36 @@ bool UGameFeaturesSubsystem::GetGameFeaturePluginDetails(const FString& PluginDe
 		}
 	}
 
-	// Read the file to a string
-	FString FileContents;
-	if (!FFileHelper::LoadFileToString(FileContents, *PluginDescriptorFilename))
+	TSharedPtr<FJsonObject> ObjectPtr;
 	{
-		UE_LOG(LogGameFeatures, Error, TEXT("UGameFeaturesSubsystem could not determine if feature was hotfixable. Failed to read file. File:%s Error:%d"), *PluginDescriptorFilename, FPlatformMisc::GetLastError());
-		return false;
-	}
+#if WITH_EDITOR
+		// In the editor we already have the plugin JSON cached
+		TSharedPtr<IPlugin> Plugin = IPluginManager::Get().FindPlugin(FPathViews::GetBaseFilename(PluginDescriptorFilename));
+		if (Plugin)
+		{
+			ObjectPtr = Plugin->GetDescriptorJson();
+		}
+		else
+#endif // WITH_EDITOR
+		{
+			// Read the file to a string
+			FString FileContents;
+			if (!FFileHelper::LoadFileToString(FileContents, *PluginDescriptorFilename))
+			{
+				UE_LOG(LogGameFeatures, Error, TEXT("UGameFeaturesSubsystem could not determine if feature was hotfixable. Failed to read file. File:%s Error:%d"), *PluginDescriptorFilename, FPlatformMisc::GetLastError());
+				return false;
+			}
 
-	// Deserialize a JSON object from the string
-	TSharedPtr< FJsonObject > ObjectPtr;
-	TSharedRef< TJsonReader<> > Reader = TJsonReaderFactory<>::Create(FileContents);
-	if (!FJsonSerializer::Deserialize(Reader, ObjectPtr) || !ObjectPtr.IsValid())
-	{
-		UE_LOG(LogGameFeatures, Error, TEXT("UGameFeaturesSubsystem could not determine if feature was hotfixable. Json invalid. File:%s. Error:%s"), *PluginDescriptorFilename, *Reader->GetErrorMessage());
-		return false;
+			// Deserialize a JSON object from the string	
+			TSharedRef< TJsonReader<> > Reader = TJsonReaderFactory<>::Create(FileContents);
+			if (!FJsonSerializer::Deserialize(Reader, ObjectPtr) || !ObjectPtr.IsValid())
+			{
+				UE_LOG(LogGameFeatures, Error, TEXT("UGameFeaturesSubsystem could not determine if feature was hotfixable. Json invalid. File:%s. Error:%s"), *PluginDescriptorFilename, *Reader->GetErrorMessage());
+				return false;
+			}
+		}
 	}
 
-	//@TODO: When we properly support downloaded plugins, will need to determine this
-	const bool bIsBuiltInPlugin = true;
-
 	// Read the properties
 
 	// Hotfixable. If it is not specified, then we assume it is
@@ -1437,7 +1454,7 @@ bool UGameFeaturesSubsystem::GetGameFeaturePluginDetails(const FString& PluginDe
 	ObjectPtr->TryGetBoolField(TEXT("Hotfixable"), OutPluginDetails.bHotfixable);
 
 	// Determine the initial plugin state
-	OutPluginDetails.BuiltInAutoState = bIsBuiltInPlugin ? DetermineBuiltInInitialFeatureState(ObjectPtr, PluginDescriptorFilename) : EBuiltInAutoState::Installed;
+	OutPluginDetails.BuiltInAutoState = DetermineBuiltInInitialFeatureState(ObjectPtr, PluginDescriptorFilename);
 
 	// Read any additional metadata the policy might want to consume (e.g., a release version number)
 	for (const FString& ExtraKey : GetDefault<UGameFeaturesSubsystemSettings>()->AdditionalPluginMetadataKeys)
@@ -1472,28 +1489,14 @@ bool UGameFeaturesSubsystem::GetGameFeaturePluginDetails(const FString& PluginDe
 						ElementObject->TryGetStringField(NameField, DependencyName);
 						if (!DependencyName.IsEmpty())
 						{
-							TSharedPtr<IPlugin> DependencyPlugin = IPluginManager::Get().FindPlugin(DependencyName);
-							if (DependencyPlugin.IsValid())
+							TValueOrError<FString, FString> ResolvedDepResult = GameSpecificPolicies->ReslovePluginDependency(PluginURL, DependencyName);
+							if (ResolvedDepResult.HasError())
 							{
-								FString DependencyURL;
-								if (!GetPluginURLByName(DependencyPlugin->GetName(), DependencyURL))
-								{
-									if (!DependencyPlugin->GetDescriptorFileName().IsEmpty() &&
-										GetDefault<UGameFeaturesSubsystemSettings>()->IsValidGameFeaturePlugin(FPaths::ConvertRelativePathToFull(DependencyPlugin->GetDescriptorFileName())) &&
-										FPaths::FileExists(DependencyPlugin->GetDescriptorFileName()))
-									{
-										DependencyURL = GetPluginURL_FileProtocol(DependencyPlugin->GetDescriptorFileName());
-									}
-								}
-
-								if (!DependencyURL.IsEmpty())
-								{
-									OutPluginDetails.PluginDependencies.Add(DependencyURL);
-								}
+								UE_LOG(LogGameFeatures, Error, TEXT("Game feature plugin '%s' has unknown dependency '%s' [%s]."), *PluginDescriptorFilename, *DependencyName, *ResolvedDepResult.GetError());
 							}
-							else
+							else if (ResolvedDepResult.HasValue() && !ResolvedDepResult.GetValue().IsEmpty()) // Dependency may not be a GFP
 							{
-								UE_LOG(LogGameFeatures, Display, TEXT("Game feature plugin '%s' has unknown dependency '%s'."), *PluginDescriptorFilename, *DependencyName);
+								OutPluginDetails.PluginDependencies.Add(ResolvedDepResult.StealValue());
 							}
 						}
 					}
@@ -1509,6 +1512,11 @@ bool UGameFeaturesSubsystem::GetGameFeaturePluginDetails(const FString& PluginDe
 	return true;
 }
 
+void UGameFeaturesSubsystem::PruneCachedGameFeaturePluginDetails(const FString& PluginURL, const FString& PluginDescriptorFilename) const
+{
+	CachedPluginDetailsByFilename.Remove(PluginDescriptorFilename);
+}
+
 UGameFeaturePluginStateMachine* UGameFeaturesSubsystem::FindGameFeaturePluginStateMachineByPluginName(const FString& PluginName) const
 {
 	for (auto StateMachineIt = GameFeaturePluginStateMachines.CreateConstIterator(); StateMachineIt; ++StateMachineIt)
@@ -1671,10 +1679,10 @@ void UGameFeaturesSubsystem::FinishTermination(UGameFeaturePluginStateMachine* M
 	TerminalGameFeaturePluginStateMachines.RemoveSwap(Machine);
 }
 
-bool UGameFeaturesSubsystem::FindOrCreatePluginDependencyStateMachines(const FString& PluginFilename, TArray<UGameFeaturePluginStateMachine*>& OutDependencyMachines)
+bool UGameFeaturesSubsystem::FindOrCreatePluginDependencyStateMachines(const FString& PluginURL, const FString& PluginFilename, TArray<UGameFeaturePluginStateMachine*>& OutDependencyMachines)
 {
 	FGameFeaturePluginDetails Details;
-	if (GetGameFeaturePluginDetails(PluginFilename, Details))
+	if (GetGameFeaturePluginDetails(PluginURL, PluginFilename, Details))
 	{
 		for (const FString& DependencyURL : Details.PluginDependencies)
 		{
diff --git a/Engine/Plugins/Experimental/GameFeatures/Source/GameFeatures/Public/GameFeaturesProjectPolicies.h b/Engine/Plugins/Experimental/GameFeatures/Source/GameFeatures/Public/GameFeaturesProjectPolicies.h
--- a/Engine/Plugins/Experimental/GameFeatures/Source/GameFeatures/Public/GameFeaturesProjectPolicies.h
+++ b/Engine/Plugins/Experimental/GameFeatures/Source/GameFeatures/Public/GameFeaturesProjectPolicies.h
@@ -41,6 +41,11 @@ public:
 	// (e.g., when doing a fast cook a game might want to disable some or all game feature plugins)
 	virtual bool IsPluginAllowed(const FString& PluginURL) const { return true; }
 
+	// Called to resolve plugin dependencies, OutDependencyURL will contain the GFP URL if this dependency is a GFP.
+	// This may be called with file protocol for built-in plugins in some cases, even if a different protocol is used at runtime.
+	// returns The dependency URL or an error if the dependency could not be resolved
+	virtual TValueOrError<FString, FString> ReslovePluginDependency(const FString& PluginURL, const FString& DependencyName) const;
+
 	// Called by code that explicitly wants to load a specific plugin
 	// (e.g., when using a fast cook a game might want to allow explicitly loaded game feature plugins)
 	virtual void ExplicitLoadGameFeaturePlugin(const FString& PluginURL, const FGameFeaturePluginLoadComplete& CompleteDelegate, const bool bActivateGameFeatures);
diff --git a/Engine/Plugins/Experimental/GameFeatures/Source/GameFeatures/Public/GameFeaturesSubsystem.h b/Engine/Plugins/Experimental/GameFeatures/Source/GameFeatures/Public/GameFeaturesSubsystem.h
--- a/Engine/Plugins/Experimental/GameFeatures/Source/GameFeatures/Public/GameFeaturesSubsystem.h
+++ b/Engine/Plugins/Experimental/GameFeatures/Source/GameFeatures/Public/GameFeaturesSubsystem.h
@@ -584,7 +584,11 @@ private:
 	const UGameFeatureData* GetRegisteredDataForStateMachine(UGameFeaturePluginStateMachine* GFSM) const;
 
 	/** Gets relevant properties out of a uplugin file */
-	bool GetGameFeaturePluginDetails(const FString& PluginDescriptorFilename, struct FGameFeaturePluginDetails& OutPluginDetails) const;
+	bool GetGameFeaturePluginDetails(const FString& PluginURL, const FString& PluginDescriptorFilename, struct FGameFeaturePluginDetails& OutPluginDetails) const;
+
+	/** Prunes any cached GFP details */
+	void PruneCachedGameFeaturePluginDetails(const FString& PluginURL, const FString& PluginDescriptorFilename) const;
+	friend struct FGameFeaturePluginState_Unmounting;
 
 	/** Gets the state machine associated with the specified plugin name */
 	UGameFeaturePluginStateMachine* FindGameFeaturePluginStateMachineByPluginName(const FString& PluginName) const;
@@ -612,7 +616,7 @@ private:
 	friend class UGameFeaturePluginStateMachine;
 
 	/** Handler for when a state machine requests its dependencies. Returns false if the dependencies could not be read */
-	bool FindOrCreatePluginDependencyStateMachines(const FString& PluginFilename, TArray<UGameFeaturePluginStateMachine*>& OutDependencyMachines);
+	bool FindOrCreatePluginDependencyStateMachines(const FString& PluginURL, const FString& PluginFilename, TArray<UGameFeaturePluginStateMachine*>& OutDependencyMachines);
 	friend struct FGameFeaturePluginState_WaitingForDependencies;
 
 	/** Handle 'ListGameFeaturePlugins' console command */
