Extending the game feature plugin wizard to support an optional subfolder creation path that is defined by a template. All new plugins derived from the template will be created in this Plugins/GameFeatures/<optional subfolder>.

diff --git a/Engine/Plugins/Experimental/GameFeatures/Source/GameFeaturesEditor/Private/GameFeaturesEditorModule.cpp b/Engine/Plugins/Experimental/GameFeatures/Source/GameFeaturesEditor/Private/GameFeaturesEditorModule.cpp
--- a/Engine/Plugins/Experimental/GameFeatures/Source/GameFeaturesEditor/Private/GameFeaturesEditorModule.cpp
+++ b/Engine/Plugins/Experimental/GameFeatures/Source/GameFeaturesEditor/Private/GameFeaturesEditorModule.cpp
@@ -34,12 +34,13 @@
 
 struct FGameFeaturePluginTemplateDescription : public FPluginTemplateDescription
 {
-	FGameFeaturePluginTemplateDescription(FText InName, FText InDescription, FString InOnDiskPath
+	FGameFeaturePluginTemplateDescription(FText InName, FText InDescription, FString InOnDiskPath, FString InDefaultSubfolder
 		, TSubclassOf<UGameFeatureData> GameFeatureDataClassOverride, FString GameFeatureDataNameOverride, EPluginEnabledByDefault InEnabledByDefault)
 		: FPluginTemplateDescription(InName, InDescription, InOnDiskPath, /*bCanContainContent=*/ true, EHostType::Runtime)
 	{
 		SortPriority = 10;
 		bCanBePlacedInEngine = false;
+		DefaultSubfolder = InDefaultSubfolder;
 		GameFeatureDataName = !GameFeatureDataNameOverride.IsEmpty() ? GameFeatureDataNameOverride : FString();
 		GameFeatureDataClass = GameFeatureDataClassOverride != nullptr ? GameFeatureDataClassOverride : TSubclassOf<UGameFeatureData>(UGameFeatureData::StaticClass());
 		PluginEnabledByDefault = InEnabledByDefault;
@@ -128,6 +129,13 @@ struct FGameFeaturePluginTemplateDescription : public FPluginTemplateDescription
 	FString GetGameFeatureRoot() const
 	{
 		FString Result = IFileManager::Get().ConvertToAbsolutePathForExternalAppForWrite(*(FPaths::ProjectPluginsDir() / TEXT("GameFeatures/")));
+
+		// Append the optional subfolder if specified.
+		if (!DefaultSubfolder.IsEmpty())
+		{
+			Result /= DefaultSubfolder + TEXT("/");
+		}
+
 		FPaths::MakePlatformFilename(Result);
 		return Result;
 	}
@@ -138,6 +146,7 @@ struct FGameFeaturePluginTemplateDescription : public FPluginTemplateDescription
 		return GetDefault<UGameFeaturesSubsystemSettings>()->IsValidGameFeaturePlugin(ConvertedPath);
 	}
 
+	FString DefaultSubfolder;
 	TSubclassOf<UGameFeatureData> GameFeatureDataClass;
 	FString GameFeatureDataName;
 	EPluginEnabledByDefault PluginEnabledByDefault = EPluginEnabledByDefault::Disabled;
@@ -230,6 +239,7 @@ class FGameFeaturesEditorModule : public FDefaultModuleImpl
 					PluginTemplate.Label,
 					PluginTemplate.Description,
 					PluginTemplate.Path.Path,
+					PluginTemplate.DefaultSubfolder,
 					PluginTemplate.DefaultGameFeatureDataClass,
 					PluginTemplate.DefaultGameFeatureDataName,
 					PluginTemplate.bIsEnabledByDefault ? EPluginEnabledByDefault::Enabled : EPluginEnabledByDefault::Disabled)));
diff --git a/Engine/Plugins/Experimental/GameFeatures/Source/GameFeaturesEditor/Private/GameFeaturesEditorSettings.h b/Engine/Plugins/Experimental/GameFeatures/Source/GameFeaturesEditor/Private/GameFeaturesEditorSettings.h
--- a/Engine/Plugins/Experimental/GameFeatures/Source/GameFeaturesEditor/Private/GameFeaturesEditorSettings.h
+++ b/Engine/Plugins/Experimental/GameFeatures/Source/GameFeaturesEditor/Private/GameFeaturesEditorSettings.h
@@ -30,6 +30,10 @@ struct FPluginTemplateData
 	UPROPERTY(EditAnywhere, Category = PluginTemplate)
 	FText Description;
 
+	/** Optional sub folder that new plugins will be created in. */
+	UPROPERTY(EditAnywhere, Category = PluginTemplate)
+	FString DefaultSubfolder;
+
 	/** The default class of game feature data to create for new game feature plugins (if not set, UGameFeatureData will be used) */
 	UPROPERTY(config, EditAnywhere, Category = Plugins)
 	TSubclassOf<UGameFeatureData> DefaultGameFeatureDataClass;
