Only modify config files if the project has a name, which avoids modifying default engine ini files from the project launcher

PR #10337: Avoid creating default config for project browser.
(Contributed by [REDACTED_PERSON])

diff --git a/Engine/Plugins/Experimental/GameFeatures/Source/GameFeaturesEditor/Private/GameFeaturesEditorModule.cpp b/Engine/Plugins/Experimental/GameFeatures/Source/GameFeaturesEditor/Private/GameFeaturesEditorModule.cpp
--- a/Engine/Plugins/Experimental/GameFeatures/Source/GameFeaturesEditor/Private/GameFeaturesEditorModule.cpp
+++ b/Engine/Plugins/Experimental/GameFeatures/Source/GameFeaturesEditor/Private/GameFeaturesEditorModule.cpp
@@ -372,7 +372,7 @@ class FGameFeaturesEditorModule : public FDefaultModuleImpl
 		// Make sure the game has the appropriate asset manager configuration or we won't be able to load game feature data assets
 		FPrimaryAssetId DummyGameFeatureDataAssetId(UGameFeatureData::StaticClass()->GetFName(), NAME_None);
 		FPrimaryAssetRules GameDataRules = UAssetManager::Get().GetPrimaryAssetRules(DummyGameFeatureDataAssetId);
-		if (GameDataRules.IsDefault())
+		if (FApp::HasProjectName() && GameDataRules.IsDefault())
 		{
 			FMessageLog("LoadErrors").Error()
 				->AddToken(FTextToken::Create(FText::Format(NSLOCTEXT("GameFeatures", "MissingRuleForGameFeatureData", "Asset Manager settings do not include an entry for assets of type {0}, which is required for game feature plugins to function."), FText::FromName(UGameFeatureData::StaticClass()->GetFName()))))
diff --git a/Engine/Plugins/Experimental/Water/Source/Editor/Private/WaterEditorModule.cpp b/Engine/Plugins/Experimental/Water/Source/Editor/Private/WaterEditorModule.cpp
--- a/Engine/Plugins/Experimental/Water/Source/Editor/Private/WaterEditorModule.cpp
+++ b/Engine/Plugins/Experimental/Water/Source/Editor/Private/WaterEditorModule.cpp
@@ -343,7 +343,7 @@ void FWaterEditorModule::AddWaterCollisionProfile()
 	// Make sure WaterCollisionProfileName is added to Engine's collision profiles
 	const FName WaterCollisionProfileName = GetDefault<UWaterRuntimeSettings>()->GetDefaultWaterCollisionProfileName();
 	FCollisionResponseTemplate WaterBodyCollisionProfile;
-	if (!UCollisionProfile::Get()->GetProfileTemplate(WaterCollisionProfileName, WaterBodyCollisionProfile))
+	if (FApp::HasProjectName() && !UCollisionProfile::Get()->GetProfileTemplate(WaterCollisionProfileName, WaterBodyCollisionProfile))
 	{
 		WaterBodyCollisionProfile.Name = WaterCollisionProfileName;
 		WaterBodyCollisionProfile.CollisionEnabled = ECollisionEnabled::QueryOnly;
