Fix crash in UnloadPackages unloading game feature plugins during PIE

diff --git a/Engine/Plugins/Experimental/GameFeatures/Source/GameFeatures/Private/GameFeaturePluginStateMachine.cpp b/Engine/Plugins/Experimental/GameFeatures/Source/GameFeatures/Private/GameFeaturePluginStateMachine.cpp
--- a/Engine/Plugins/Experimental/GameFeatures/Source/GameFeatures/Private/GameFeaturePluginStateMachine.cpp
+++ b/Engine/Plugins/Experimental/GameFeatures/Source/GameFeatures/Private/GameFeaturePluginStateMachine.cpp
@@ -34,6 +34,7 @@
 
 #if WITH_EDITOR
 #include "PluginUtils.h"
+#include "Misc/App.h"
 #endif //if WITH_EDITOR
 
 namespace UE::GameFeatures
@@ -2132,8 +2133,18 @@ struct FGameFeaturePluginState_Unregistering : public FGameFeaturePluginState
 #if WITH_EDITOR
 		// This will properly unload any plugin asset that could be opened in the editor
 		// and ensure standalone packages get unloaded as well
-		bRequestedUnloadPluginAssets = true;
-		UE::GameFeatures::ScheduleUnloadPluginAssets(StateProperties.PluginName, StateProperties.OnRequestUpdateStateMachine);
+		if (FApp::IsGame())
+		{
+			verify(FPluginUtils::UnloadPluginAssets(StateProperties.PluginName));
+
+			bRequestedGC = true;
+			GarbageCollectAndUpdateStateMachineDeferred(/*bWaitForCleanupQueue=*/true);
+		}
+		else
+		{
+			bRequestedUnloadPluginAssets = true;
+			UE::GameFeatures::ScheduleUnloadPluginAssets(StateProperties.PluginName, StateProperties.OnRequestUpdateStateMachine);
+		}
 #else
 		// never mark anything as garbage in the editor, there be dwagons here. 
 		if (UE::GameFeatures::MarkPluginAsGarbageOnUnregister && !UE::GameFeatures::ShouldSkipVerify(StateProperties.PluginName))
