Fix RebuildObjectReferenceList being called twice per module.

diff --git a/Engine/Plugins/Experimental/GameFeatures/Source/GameFeatures/Private/GameFeaturesSubsystem.cpp b/Engine/Plugins/Experimental/GameFeatures/Source/GameFeatures/Private/GameFeaturesSubsystem.cpp
--- a/Engine/Plugins/Experimental/GameFeatures/Source/GameFeatures/Private/GameFeaturesSubsystem.cpp
+++ b/Engine/Plugins/Experimental/GameFeatures/Source/GameFeatures/Private/GameFeaturesSubsystem.cpp
@@ -491,6 +491,8 @@ void UGameFeaturesSubsystem::AddGameFeatureToAssetManager(const UGameFeatureData
 	UAssetManager& LocalAssetManager = UAssetManager::Get();
 	IAssetRegistry& LocalAssetRegistry = LocalAssetManager.GetAssetRegistry();
 
+	LocalAssetManager.PushBulkScanning();
+
 	// Add the GameFeatureData itself to the primary asset list
 #if WITH_EDITOR
 	// In the editor, we may not have scanned the FAssetData yet if during startup, but that is fine because we can gather bundles from the object itself, so just create the FAssetData from the object
@@ -500,7 +502,7 @@ void UGameFeaturesSubsystem::AddGameFeatureToAssetManager(const UGameFeatureData
 	LocalAssetManager.RegisterSpecificPrimaryAsset(GameFeatureToAdd->GetPrimaryAssetId(), LocalAssetRegistry.GetAssetByObjectPath(FSoftObjectPath(GameFeatureToAdd), true));
 #endif // WITH_EDITOR
 
-	// @TODO: HACK - There is no guarantee that the plugin mount point was added before inte initial asset scan.
+	// @TODO: HACK - There is no guarantee that the plugin mount point was added before the initial asset scan.
 	// If not, ScanPathsForPrimaryAssets will fail to find primary assets without a syncronous scan.
 	// A proper fix for this would be to handle all the primary asset discovery internally ins the asset manager 
 	// instead of doing it here.
@@ -509,8 +511,6 @@ void UGameFeaturesSubsystem::AddGameFeatureToAssetManager(const UGameFeatureData
 	// as long as LocalAssetRegistry.IsLoadingAssets() is true.
 	const bool bForceSynchronousScan = !LocalAssetRegistry.IsLoadingAssets();
 
-	LocalAssetManager.PushBulkScanning();
-
 	for (FPrimaryAssetTypeInfo TypeInfo : GameFeatureToAdd->GetPrimaryAssetTypesToScan())
 	{
 		// @TODO: we shouldn't be accessing private data here. Need a better way to do this
