Added the ability to load editor ini files.

diff --git a/Engine/Plugins/Experimental/GameFeatures/Source/GameFeatures/Private/GameFeatureData.cpp b/Engine/Plugins/Experimental/GameFeatures/Source/GameFeatures/Private/GameFeatureData.cpp
--- a/Engine/Plugins/Experimental/GameFeatures/Source/GameFeatures/Private/GameFeatureData.cpp
+++ b/Engine/Plugins/Experimental/GameFeatures/Source/GameFeatures/Private/GameFeatureData.cpp
@@ -111,7 +111,12 @@ void UGameFeatureData::InitializeHierarchicalPluginIniFiles(const FString& Plugi
 	const bool bWriteDestIni = false;
 
 	// @todo: Likely we need to track the diffs this config caused and/or store versions/layers in order to unwind settings during unloading/deactivation
-	TArray<FString> IniNamesToLoad = { TEXT("Input"), TEXT("Game"), TEXT("Engine") };
+	TArray<FString> IniNamesToLoad = { 
+		TEXT("Input"), TEXT("Game"), TEXT("Engine")
+#if UE_EDITOR
+		, TEXT("Editor")
+#endif
+	};
 	for (const FString& IniName : IniNamesToLoad)
 	{
 		const FString PluginIniName = PluginName + IniName;
