Additional logging for when GFP plugin dependencies fail.

diff --git a/Engine/Plugins/Experimental/GameFeatures/Source/GameFeatures/Private/GameFeaturePluginStateMachine.cpp b/Engine/Plugins/Experimental/GameFeatures/Source/GameFeatures/Private/GameFeaturePluginStateMachine.cpp
--- a/Engine/Plugins/Experimental/GameFeatures/Source/GameFeatures/Private/GameFeaturePluginStateMachine.cpp
+++ b/Engine/Plugins/Experimental/GameFeatures/Source/GameFeatures/Private/GameFeaturePluginStateMachine.cpp
@@ -1969,6 +1969,7 @@ struct FGameFeaturePluginState_WaitingForDependencies : public FGameFeaturePlugi
 
 			if (Pair.Value.HasError())
 			{
+				UE_LOG(LogGameFeatures, Error, TEXT("Dependency %s failed to load with error %s"), *RemainingDependency->GetPluginURL(), *Pair.Value.GetError());
 				StateStatus.SetTransitionError(EGameFeaturePluginState::ErrorWaitingForDependencies, GetErrorResult(TEXT("Failed_Dependency_Register")));
 				return;
 			}
