Always inject content when running the cook commandlet
Tempory fix until using a separate gamefeature action
Not integrated in main.

diff --git a/Engine/Plugins/Experimental/GameFeatures/Source/GameFeatures/Private/GameFeatureAction_AddWPContent.cpp b/Engine/Plugins/Experimental/GameFeatures/Source/GameFeatures/Private/GameFeatureAction_AddWPContent.cpp
--- a/Engine/Plugins/Experimental/GameFeatures/Source/GameFeatures/Private/GameFeatureAction_AddWPContent.cpp
+++ b/Engine/Plugins/Experimental/GameFeatures/Source/GameFeatures/Private/GameFeatureAction_AddWPContent.cpp
@@ -29,6 +29,13 @@ void UGameFeatureAction_AddWPContent::OnGameFeatureRegistering()
 	Super::OnGameFeatureRegistering();
 	ContentBundleClient = FContentBundleClient::CreateClient(ContentBundleDescriptor, GetTypedOuter<UGameFeatureData>()->GetName());
 	UE_CLOG(ContentBundleClient == nullptr, LogGameFeatures, Error, TEXT("OnGameFeatureRegistering %s: Failed to create a content bundle client for %s"), *GetPathName(), *ContentBundleDescriptor->GetDisplayName())
+
+#if WITH_EDITOR
+	if (IsRunningCookCommandlet() && ContentBundleClient != nullptr)
+	{
+		ContentBundleClient->RequestContentInjection();
+	}
+#endif 
 }
 
 void UGameFeatureAction_AddWPContent::OnGameFeatureUnregistering()
