GFP
- add verbose logging for number of dependencies
- add an ensure to dependency callback

diff --git a/Engine/Plugins/Experimental/GameFeatures/Source/GameFeatures/Private/GameFeaturePluginStateMachine.cpp b/Engine/Plugins/Experimental/GameFeatures/Source/GameFeatures/Private/GameFeaturePluginStateMachine.cpp
--- a/Engine/Plugins/Experimental/GameFeatures/Source/GameFeatures/Private/GameFeaturePluginStateMachine.cpp
+++ b/Engine/Plugins/Experimental/GameFeatures/Source/GameFeatures/Private/GameFeaturePluginStateMachine.cpp
@@ -1949,6 +1949,8 @@ struct FGameFeaturePluginState_WaitingForDependencies : public FGameFeaturePlugi
 			
 			bRequestedDependencies = true;
 
+			UE_CLOG(Dependencies.Num() > 0, LogGameFeatures, Verbose, TEXT("Found %i dependencies for %s"), Dependencies.Num(), *StateProperties.PluginName);
+
 			RemainingDependencies.Reserve(Dependencies.Num());
 			for (UGameFeaturePluginStateMachine* Dependency : Dependencies)
 			{	
@@ -2042,7 +2044,7 @@ struct FGameFeaturePluginState_WaitingForDependencies : public FGameFeaturePlugi
 			return Pair.Key == Dependency;
 		});
 
-		if (Index != INDEX_NONE)
+		if (ensure(Index != INDEX_NONE))
 		{
 			FDepResultPair& FoundDep = RemainingDependencies[Index];
 			
