GameFeaturePlugin improvements:
- Fixed error in  operator> for FGameFeaturePluginStateRange when comparing 2 ranges.
- Added new Uninstalled destination state signifying we have finished uninstalling data. This gives us a separate target for uninstalling besides Terminal if we need to explicitly target that state.
- Changed UninstallGameFeaturePlugin to go through targetting Uninstall first and then terminating the plugin after reaching Uninstalled. This is to allow you to uninstall data for GameFeaturePlugins that were installed on a previous run of the App without first requesting they get to at least StatusKnown (otherwise we fallthrough UnknownStatus -> Terminal without uninstalling).

diff --git a/Engine/Plugins/Experimental/GameFeatures/Source/GameFeatures/Private/GameFeaturePluginStateMachine.cpp b/Engine/Plugins/Experimental/GameFeatures/Source/GameFeatures/Private/GameFeaturePluginStateMachine.cpp
--- a/Engine/Plugins/Experimental/GameFeatures/Source/GameFeatures/Private/GameFeaturePluginStateMachine.cpp
+++ b/Engine/Plugins/Experimental/GameFeatures/Source/GameFeatures/Private/GameFeaturePluginStateMachine.cpp
@@ -1136,6 +1136,23 @@ struct FBaseDataReleaseGameFeaturePluginState : public FGameFeaturePluginState
 	virtual EGameFeaturePluginState GetFailureTransitionState() const = 0;
 };
 
+struct FGameFeaturePluginState_Uninstalled : public FDestinationGameFeaturePluginState
+{
+	FGameFeaturePluginState_Uninstalled(FGameFeaturePluginStateMachineProperties& InStateProperties) : FDestinationGameFeaturePluginState(InStateProperties) {}
+
+	virtual void UpdateState(FGameFeaturePluginStateStatus& StateStatus) override
+	{
+		if (StateProperties.Destination < EGameFeaturePluginState::Uninstalled)
+		{
+			StateStatus.SetTransition(EGameFeaturePluginState::Terminal);
+		}
+		else if (StateProperties.Destination > EGameFeaturePluginState::Uninstalled)
+		{
+			StateStatus.SetTransition(EGameFeaturePluginState::CheckingStatus);
+		}
+	}
+};
+
 struct FGameFeaturePluginState_Uninstalling : public FBaseDataReleaseGameFeaturePluginState
 {
 	FGameFeaturePluginState_Uninstalling(FGameFeaturePluginStateMachineProperties& InStateProperties)
@@ -1144,7 +1161,7 @@ struct FGameFeaturePluginState_Uninstalling : public FBaseDataReleaseGameFeature
 
 	virtual EGameFeaturePluginState GetSuccessTransitionState() const override
 	{
-		return EGameFeaturePluginState::Terminal;
+		return EGameFeaturePluginState::Uninstalled;
 	}
 	
 	virtual EGameFeaturePluginState GetFailureTransitionState() const override
diff --git a/Engine/Plugins/Experimental/GameFeatures/Source/GameFeatures/Private/GameFeaturePluginStateMachine.h b/Engine/Plugins/Experimental/GameFeatures/Source/GameFeatures/Private/GameFeaturePluginStateMachine.h
--- a/Engine/Plugins/Experimental/GameFeatures/Source/GameFeatures/Private/GameFeaturePluginStateMachine.h
+++ b/Engine/Plugins/Experimental/GameFeatures/Source/GameFeatures/Private/GameFeaturePluginStateMachine.h
@@ -39,6 +39,12 @@ Transition states are expected to transition the machine to another state after
         |           ^                      ^        |               |          |
         |           |                      |        +-------+-------+          |
         |           |                      |                |                  |
+        |    +------+-------+              |                |                  |
+        |    |      *       |              |                |                  |
+        |    | Uninstalled  +--------------~--------------->|                  |
+        |    |              |              |                |                  |
+        |    +------^-------+              |                |                  |
+        |           |                      |                |                  |
         |    +------+-------+    *---------+---------+      |                  |
         |    |              |    |         !         |      |                  |
         |    | Uninstalling <----> ErrorUninstalling |      |                  |
@@ -185,7 +191,7 @@ struct FGameFeaturePluginStateRange
 
 	bool operator==(const FGameFeaturePluginStateRange& Other) const { return MinState == Other.MinState && MaxState == Other.MaxState; }
 	bool operator<(const FGameFeaturePluginStateRange& Other) const { return MaxState < Other.MinState; }
-	bool operator>(const FGameFeaturePluginStateRange& Other) const { return MinState < Other.MaxState; }
+	bool operator>(const FGameFeaturePluginStateRange& Other) const { return MinState > Other.MaxState; }
 };
 
 inline bool operator<(EGameFeaturePluginState State, const FGameFeaturePluginStateRange& StateRange)
diff --git a/Engine/Plugins/Experimental/GameFeatures/Source/GameFeatures/Private/GameFeaturesSubsystem.cpp b/Engine/Plugins/Experimental/GameFeatures/Source/GameFeatures/Private/GameFeaturesSubsystem.cpp
--- a/Engine/Plugins/Experimental/GameFeatures/Source/GameFeatures/Private/GameFeaturesSubsystem.cpp
+++ b/Engine/Plugins/Experimental/GameFeatures/Source/GameFeatures/Private/GameFeaturesSubsystem.cpp
@@ -1074,6 +1074,9 @@ void UGameFeaturesSubsystem::ReleaseGameFeaturePlugin(const FString& PluginURL,
 
 void UGameFeaturesSubsystem::UninstallGameFeaturePlugin(const FString& PluginURL, const FGameFeaturePluginUninstallComplete& CompleteDelegate)
 {
+	//FindOrCreate so that we can make sure we uninstall data for plugins that were installed on a previous application run
+	//but have not yet been requested on this application run and so are not yet in the plugin list but might have data on disk
+	//to uninstall
 	UGameFeaturePluginStateMachine* StateMachine = FindOrCreateGameFeaturePluginStateMachine(PluginURL);
 	check(StateMachine);
 
@@ -1100,7 +1103,48 @@ void UGameFeaturesSubsystem::UninstallGameFeaturePlugin(const FString& PluginURL
 		}
 	}
 
-	TerminateGameFeaturePlugin(PluginURLForTerminate, CompleteDelegate);
+	//Weird flow here because we need to do a few tasks asynchronously
+	// Update URL   -->    Call to set destination to Uninstall --> After we get to Uninstall go to Terminate
+	// (Called Directly)	   (UninstallTransitionLambda)              (StartTerminateLambda)
+
+	// THIRD:
+	//Lambda that will kick off the actual Terminate after we successfully transition to Uninstalled state
+	const FGameFeaturePluginTerminateComplete StartTerminateLambda = FGameFeaturePluginTerminateComplete::CreateWeakLambda(this,
+		[this, PluginURLForTerminate, CompleteDelegate](const UE::GameFeatures::FResult& Result)
+		{
+			if (Result.HasValue())
+			{
+				TerminateGameFeaturePlugin(PluginURLForTerminate, CompleteDelegate);
+			}
+			//If we failed just bubble error up
+			else
+			{
+				CompleteDelegate.ExecuteIfBound(Result);
+			}
+		});
+
+	// SECOND:
+	//Lambda that will kick off the Uninstall destination after updating our URL if necessary
+	const FGameFeaturePluginUpdateURLComplete UninstallTransitionLambda = FGameFeaturePluginUpdateURLComplete::CreateWeakLambda(this,
+		[this, PluginURLForTerminate, StartTerminateLambda](const UE::GameFeatures::FResult& Result)
+		{
+			UGameFeaturePluginStateMachine* StateMachine = FindGameFeaturePluginStateMachine(PluginURLForTerminate);
+			check(StateMachine);
+			ChangeGameFeatureDestination(StateMachine, FGameFeaturePluginStateRange(EGameFeaturePluginState::Uninstalled), StartTerminateLambda);
+		});
+
+
+	// FIRST:
+	//If we need to update our URLData, try to do that first before starting the Uninstall. This allows us to update
+	//URL Metadata flags that might be important on the way to Terminal if they are changed. EX: FInstallBundlePluginProtocolMetaData::bUninstallBeforeTerminate
+	if (ShouldUpdatePluginURLData(PluginURLForTerminate))
+	{
+		UpdateGameFeaturePluginURL(PluginURLForTerminate, UninstallTransitionLambda);
+	}
+	else
+	{
+		UninstallTransitionLambda.Execute(MakeValue());
+	}
 }
 
 void UGameFeaturesSubsystem::TerminateGameFeaturePlugin(const FString& PluginURL)
diff --git a/Engine/Plugins/Experimental/GameFeatures/Source/GameFeatures/Public/GameFeatureTypes.h b/Engine/Plugins/Experimental/GameFeatures/Source/GameFeatures/Public/GameFeatureTypes.h
--- a/Engine/Plugins/Experimental/GameFeatures/Source/GameFeatures/Public/GameFeatureTypes.h
+++ b/Engine/Plugins/Experimental/GameFeatures/Source/GameFeatures/Public/GameFeatureTypes.h
@@ -10,7 +10,8 @@ class FString;
 	XSTATE(Uninitialized,				NSLOCTEXT("GameFeatures", "UninitializedStateDisplayName", "Uninitialized"))								/* Unset. Not yet been set up. */ \
 	XSTATE(Terminal,					NSLOCTEXT("GameFeatures", "TerminalStateDisplayName", "Terminal"))											/* Final State before removal of the state machine. */ \
 	XSTATE(UnknownStatus,				NSLOCTEXT("GameFeatures", "UnknownStatusStateDisplayName", "UnknownStatus"))								/* Initialized, but the only thing known is the URL to query status. */ \
-	XSTATE(Uninstalling,				NSLOCTEXT("GameFeatures", "UninstallingStateDisplayName", "Uninstalling"))									/* Transition state between StatusKnown -> Terminal for any plugin that can have data that needs to have local datat uninstalled. */ \
+	XSTATE(Uninstalled,					NSLOCTEXT("GameFeatures", "UninstalledStateDisplayName", "Uninstalled"))									/* All installed data for this plugin has now been uninstalled from local storage (i.e the hard drive) */ \
+	XSTATE(Uninstalling,				NSLOCTEXT("GameFeatures", "UninstallingStateDisplayName", "Uninstalling"))									/* Transition state between StatusKnown -> Terminal for any plugin that can have data that needs to have local data uninstalled. */ \
 	XSTATE(ErrorUninstalling,			NSLOCTEXT("GameFeatures", "ErrorUninstallingStateDisplayName", "ErrorUninstalling"))						/* Error state for Uninstalling -> Terminal transition. */  \
 	XSTATE(CheckingStatus,				NSLOCTEXT("GameFeatures", "CheckingStatusStateDisplayName", "CheckingStatus"))								/* Transition state UnknownStatus -> StatusKnown. The status is in the process of being queried. */ \
 	XSTATE(ErrorCheckingStatus,			NSLOCTEXT("GameFeatures", "ErrorCheckingStatusStateDisplayName", "ErrorCheckingStatus"))					/* Error state for UnknownStatus -> StatusKnown transition. */ \
diff --git a/Engine/Plugins/Experimental/GameFeatures/Source/GameFeaturesEditor/Private/SGameFeatureStateWidget.cpp b/Engine/Plugins/Experimental/GameFeatures/Source/GameFeaturesEditor/Private/SGameFeatureStateWidget.cpp
--- a/Engine/Plugins/Experimental/GameFeatures/Source/GameFeaturesEditor/Private/SGameFeatureStateWidget.cpp
+++ b/Engine/Plugins/Experimental/GameFeatures/Source/GameFeaturesEditor/Private/SGameFeatureStateWidget.cpp
@@ -65,7 +65,7 @@ FText SGameFeatureStateWidget::GetDisplayNameOfState(EGameFeaturePluginState Sta
 
 FText SGameFeatureStateWidget::GetTooltipOfState(EGameFeaturePluginState State)
 {
-	static_assert((int32)EGameFeaturePluginState::MAX == 28, "");
+	static_assert((int32)EGameFeaturePluginState::MAX == 29, "");
 
 	switch (State)
 	{
@@ -75,6 +75,8 @@ FText SGameFeatureStateWidget::GetTooltipOfState(EGameFeaturePluginState State)
 		return LOCTEXT("StateTooltip_Terminal", "Final State before removal of the state machine");
 	case EGameFeaturePluginState::UnknownStatus:
 		return LOCTEXT("StateTooltip_UnknownStatus", "Initialized, but the only thing known is the URL to query status.");
+	case EGameFeaturePluginState::Uninstalled:
+		return LOCTEXT("StateTooltip_Uninstalled", "All installed data for this plugin has now been uninstalled from local storage (i.e the hard drive)");
 	case EGameFeaturePluginState::Uninstalling:
 		return LOCTEXT("StateTooltip_Uninstalling", "Transition state between StatusKnown -> Terminal for any plugin that can have data that needs to have local datat uninstalled.");
 	case EGameFeaturePluginState::ErrorUninstalling:
