Fixing GameFeaturePlugins not respecting set FInstallBundlePluginProtocolMetaData ReleaseInstallBundleFlags during Unmount.
#tests Android Dev Client

diff --git a/Engine/Plugins/Experimental/GameFeatures/Source/GameFeatures/Private/GameFeaturePluginStateMachine.cpp b/Engine/Plugins/Experimental/GameFeatures/Source/GameFeatures/Private/GameFeaturePluginStateMachine.cpp
--- a/Engine/Plugins/Experimental/GameFeatures/Source/GameFeatures/Private/GameFeaturePluginStateMachine.cpp
+++ b/Engine/Plugins/Experimental/GameFeatures/Source/GameFeatures/Private/GameFeaturePluginStateMachine.cpp
@@ -1679,7 +1679,7 @@ struct FGameFeaturePluginState_Unmounting : public FGameFeaturePluginState
 
 		const TArray<FName>& InstallBundles = StateProperties.ProtocolMetadata.GetSubtype<FInstallBundlePluginProtocolMetaData>().InstallBundles;
 
-		EInstallBundleReleaseRequestFlags ReleaseFlags = EInstallBundleReleaseRequestFlags::None;
+		EInstallBundleReleaseRequestFlags ReleaseFlags = StateProperties.ProtocolMetadata.GetSubtype<FInstallBundlePluginProtocolMetaData>().ReleaseInstallBundleFlags;
 		TValueOrError<FInstallBundleRequestInfo, EInstallBundleResult> MaybeRequestInfo = BundleManager->RequestReleaseContent(InstallBundles, ReleaseFlags);
 
 		if (MaybeRequestInfo.HasError())
