Duplicate LocText usage fix

diff --git a/Engine/Plugins/Experimental/GameFeatures/Source/GameFeatures/Private/GameFeatureAction_AddChunkOverride.cpp b/Engine/Plugins/Experimental/GameFeatures/Source/GameFeatures/Private/GameFeatureAction_AddChunkOverride.cpp
--- a/Engine/Plugins/Experimental/GameFeatures/Source/GameFeatures/Private/GameFeatureAction_AddChunkOverride.cpp
+++ b/Engine/Plugins/Experimental/GameFeatures/Source/GameFeatures/Private/GameFeatureAction_AddChunkOverride.cpp
@@ -251,13 +251,13 @@ int32 UGameFeatureAction_AddChunkOverride::GenerateUniqueChunkId() const
 	if (NewChunkId < LowestAllowedChunkIndexForAutoGeneration)
 	{
 		UE_LOG(LogAddChunkOverride, Warning, TEXT("Autogenerated ChunkId(%d) is lower than the config specified LowestAllowedChunkIndexForAutoGeneration(%d)"), NewChunkId, LowestAllowedChunkIndexForAutoGeneration);
-		FMessageDialog::Open(EAppMsgType::Ok, LOCTEXT("GameFeatureAction_AddChunkOverride", "Autogenerated ChunkID is lower than config specified LowestAllowedChunkIndexForAutoGeneration. Please manually assign a valid Chunk Id"));
+		FMessageDialog::Open(EAppMsgType::Ok, LOCTEXT("AddChunkOverride_InvalidId", "Autogenerated ChunkID is lower than config specified LowestAllowedChunkIndexForAutoGeneration. Please manually assign a valid Chunk Id"));
 		NewChunkId = -1;
 	}
 	else if (GameFeatureAction_AddChunkOverride::ChunkIdToPluginMap.Contains(NewChunkId))
 	{
 		UE_LOG(LogAddChunkOverride, Warning, TEXT("ChunkId(%d) is in use by %s. Unable to autogenerate unique id. Lowest allowed ChunkId(%d)"), NewChunkId, *GameFeatureAction_AddChunkOverride::ChunkIdToPluginMap[NewChunkId]);
-		FMessageDialog::Open(EAppMsgType::Ok, LOCTEXT("GameFeatureAction_AddChunkOverride", "Unable to auto generate unique valid Chunk Id. Please manually assign a valid Chunk Id"));
+		FMessageDialog::Open(EAppMsgType::Ok, LOCTEXT("AddChunkOverride_UsedChunkId", "Unable to auto generate unique valid Chunk Id. Please manually assign a valid Chunk Id"));
 		NewChunkId = -1;
 	}
 
