- Fixing error in GameFeaturePlugin Uninstall change, was pulling PluginName incorrectly instead of PluginInstalledName.
- Also adding fallback behavior as if we are polling PluginInstalledName before hitting StatusKnown it won't be populated. Now can fallback to generate it from the URL
#tests Android Dev Client

diff --git a/Engine/Plugins/Experimental/GameFeatures/Source/GameFeatures/Private/GameFeaturesSubsystem.cpp b/Engine/Plugins/Experimental/GameFeatures/Source/GameFeatures/Private/GameFeaturesSubsystem.cpp
--- a/Engine/Plugins/Experimental/GameFeatures/Source/GameFeatures/Private/GameFeaturesSubsystem.cpp
+++ b/Engine/Plugins/Experimental/GameFeatures/Source/GameFeatures/Private/GameFeaturesSubsystem.cpp
@@ -1098,7 +1098,20 @@ void UGameFeaturesSubsystem::UninstallGameFeaturePlugin(const FString& PluginURL
 		if (!ProtocolMetadata.bUninstallBeforeTerminate)
 		{
 			ProtocolMetadata.bUninstallBeforeTerminate = true;
-			PluginURLForTerminate = GetPluginURL_InstallBundleProtocol(StateMachine->GetPluginName(), ProtocolMetadata);
+			
+			FString PluginFilename;
+
+			//Try and pull PluginFilename from the StateMachine first, but this may not be set yet
+			//as this may be running too early before we have parsed the URL
+			StateMachine->GetPluginFilename(PluginFilename);
+			if (PluginFilename.IsEmpty())
+			{
+				//The PluginIdentifyingString is currently just the PluginFilename so we can fallback to that
+				PluginFilename = StateMachine->GetPluginIdentifier().GetIdentifyingString();
+			}
+			check(!PluginFilename.IsEmpty());
+
+			PluginURLForTerminate = GetPluginURL_InstallBundleProtocol(PluginFilename, ProtocolMetadata);
 		}
 	}
 
