Ensure Unmounting and Releasing states in GameFeaturePlugins exclude the RemoveFilesIfPossible flag even if the user has set it, this should only be set on the Uninstalling state.
#tests Android Dev Client

diff --git a/Engine/Plugins/Experimental/GameFeatures/Source/GameFeatures/Private/GameFeaturePluginStateMachine.cpp b/Engine/Plugins/Experimental/GameFeatures/Source/GameFeatures/Private/GameFeaturePluginStateMachine.cpp
--- a/Engine/Plugins/Experimental/GameFeatures/Source/GameFeatures/Private/GameFeaturePluginStateMachine.cpp
+++ b/Engine/Plugins/Experimental/GameFeatures/Source/GameFeatures/Private/GameFeaturePluginStateMachine.cpp
@@ -1218,6 +1218,14 @@ struct FGameFeaturePluginState_Releasing : public FBaseDataReleaseGameFeaturePlu
 	{
 		return EGameFeaturePluginState::ErrorManagingData;
 	}
+
+	//Overriden to ensure we don't call with RemoveFilesIfPossible early, that should
+	//only be called during Uninstall
+	virtual EInstallBundleReleaseRequestFlags GetReleaseRequestFlags() const override
+	{
+		const EInstallBundleReleaseRequestFlags BaseFlags = FBaseDataReleaseGameFeaturePluginState::GetReleaseRequestFlags();
+		return (BaseFlags & ~(EInstallBundleReleaseRequestFlags::RemoveFilesIfPossible));
+	}
 };
 
 struct FGameFeaturePluginState_Downloading : public FGameFeaturePluginState
@@ -1680,6 +1688,9 @@ struct FGameFeaturePluginState_Unmounting : public FGameFeaturePluginState
 		const TArray<FName>& InstallBundles = StateProperties.ProtocolMetadata.GetSubtype<FInstallBundlePluginProtocolMetaData>().InstallBundles;
 
 		EInstallBundleReleaseRequestFlags ReleaseFlags = StateProperties.ProtocolMetadata.GetSubtype<FInstallBundlePluginProtocolMetaData>().ReleaseInstallBundleFlags;
+		//Make sure we don't remove files here early, that should only be done in Uninstalling
+		ReleaseFlags &= ~(EInstallBundleReleaseRequestFlags::RemoveFilesIfPossible);
+		
 		TValueOrError<FInstallBundleRequestInfo, EInstallBundleResult> MaybeRequestInfo = BundleManager->RequestReleaseContent(InstallBundles, ReleaseFlags);
 
 		if (MaybeRequestInfo.HasError())
