#UE Downgrade long plugin load time message to a warning to avoid it preventing builds from succeeding in cases where errors are not tolerated

diff --git a/Engine/Plugins/Experimental/GameFeatures/Source/GameFeatures/Private/GameFeaturesSubsystem.cpp b/Engine/Plugins/Experimental/GameFeatures/Source/GameFeatures/Private/GameFeaturesSubsystem.cpp
--- a/Engine/Plugins/Experimental/GameFeatures/Source/GameFeatures/Private/GameFeaturesSubsystem.cpp
+++ b/Engine/Plugins/Experimental/GameFeatures/Source/GameFeatures/Private/GameFeaturesSubsystem.cpp
@@ -1292,7 +1292,7 @@ public:
 				double LoadTime = Plugin.Value;
 				if (LoadTime >= UE::GameFeatures::GBuiltInPluginLoadTimeErrorThreshold)
 				{
-					UE_LOG(LogGameFeatures, Error, TEXT("%s took %.4f seconds to load. Something was done to significantly increase the load time of this plugin and it is now well outside what is acceptable. Reduce the load time to much less than %.4f seconds. Ideally, reduce the load time to less than %.4f seconds."),
+					UE_LOG(LogGameFeatures, Warning, TEXT("%s took %.4f seconds to load. Something was done to significantly increase the load time of this plugin and it is now well outside what is acceptable. Reduce the load time to much less than %.4f seconds. Ideally, reduce the load time to less than %.4f seconds."),
 						*Plugin.Key, Plugin.Value, UE::GameFeatures::GBuiltInPluginLoadTimeErrorThreshold, UE::GameFeatures::GBuiltInPluginLoadTimeReportThreshold);
 				}
 				else
