changed FShaderCodeLibraryPluginUnmountedCallback to close the shader library at the end of the render pipeline
https://p4-swarm.epicgames.net/reviews/25511457

diff --git a/Engine/Plugins/Experimental/GameFeatures/Source/GameFeatures/Private/GameFeaturePluginStateMachine.cpp b/Engine/Plugins/Experimental/GameFeatures/Source/GameFeatures/Private/GameFeaturePluginStateMachine.cpp
--- a/Engine/Plugins/Experimental/GameFeatures/Source/GameFeatures/Private/GameFeaturePluginStateMachine.cpp
+++ b/Engine/Plugins/Experimental/GameFeatures/Source/GameFeatures/Private/GameFeaturePluginStateMachine.cpp
@@ -612,42 +612,13 @@ void FGameFeaturePluginState::UpdateStateMachineDeferred(float Delay /*= 0.0f*/)
 	}), Delay);
 }
 
-void FGameFeaturePluginState::InsertCleanupObjectsFence() const
-{
-	class FLevelStreamoutFence : public FDeferredCleanupInterface
-	{
-	public:
-		FLevelStreamoutFence(const FGameFeaturePluginState& InState)
-			: State(InState)
-		{
-
-		}
-		virtual ~FLevelStreamoutFence()
-		{
-			State.UpdateStateMachineDeferred();
-		}
-
-	private:
-		const FGameFeaturePluginState& State;
-	};
-
-	BeginCleanup(new FLevelStreamoutFence(*this));
-}
-
-void FGameFeaturePluginState::GarbageCollectAndUpdateStateMachineDeferred(bool bWaitForCleanupQueue /*= false*/) const
+void FGameFeaturePluginState::GarbageCollectAndUpdateStateMachineDeferred() const
 {
 	GEngine->ForceGarbageCollection(true); // Tick Delayed
 
 	CleanupDeferredUpdateCallbacks();
 
-	if (bWaitForCleanupQueue)
-	{
-		FCoreUObjectDelegates::GetPostGarbageCollect().AddRaw(this, &FGameFeaturePluginState::InsertCleanupObjectsFence);
-	}
-	else
-	{
-		FCoreUObjectDelegates::GetPostGarbageCollect().AddRaw(this, &FGameFeaturePluginState::UpdateStateMachineDeferred, 0.0f);
-	}
+	FCoreUObjectDelegates::GetPostGarbageCollect().AddRaw(this, &FGameFeaturePluginState::UpdateStateMachineDeferred, 0.0f);
 }
 
 void FGameFeaturePluginState::MarkPluginAsGarbage(bool bMarkGameFeatureDataAsGarbage)
@@ -2141,7 +2112,7 @@ struct FGameFeaturePluginState_Unregistering : public FGameFeaturePluginState
 		if (bRequestedUnloadPluginAssets)
 		{
 			bRequestedGC = true;
-			GarbageCollectAndUpdateStateMachineDeferred(/*bWaitForCleanupQueue=*/true);
+			GarbageCollectAndUpdateStateMachineDeferred();
 			return;
 		}
 #endif //if WITH_EDITOR
@@ -2169,7 +2140,7 @@ struct FGameFeaturePluginState_Unregistering : public FGameFeaturePluginState
 			verify(FPluginUtils::UnloadPluginAssets(StateProperties.PluginName));
 
 			bRequestedGC = true;
-			GarbageCollectAndUpdateStateMachineDeferred(/*bWaitForCleanupQueue=*/true);
+			GarbageCollectAndUpdateStateMachineDeferred();
 		}
 		else
 		{
@@ -2184,7 +2155,7 @@ struct FGameFeaturePluginState_Unregistering : public FGameFeaturePluginState
 		}
 
 		bRequestedGC = true;
-		GarbageCollectAndUpdateStateMachineDeferred(/*bWaitForCleanupQueue=*/true);
+		GarbageCollectAndUpdateStateMachineDeferred();
 #endif
 	}
 };
diff --git a/Engine/Plugins/Experimental/GameFeatures/Source/GameFeatures/Private/GameFeaturePluginStateMachine.h b/Engine/Plugins/Experimental/GameFeatures/Source/GameFeatures/Private/GameFeaturePluginStateMachine.h
--- a/Engine/Plugins/Experimental/GameFeatures/Source/GameFeatures/Private/GameFeaturePluginStateMachine.h
+++ b/Engine/Plugins/Experimental/GameFeatures/Source/GameFeatures/Private/GameFeaturePluginStateMachine.h
@@ -352,9 +352,8 @@ struct FGameFeaturePluginState
 	/** The common properties that can be accessed by the states of the state machine */
 	FGameFeaturePluginStateMachineProperties& StateProperties;
 
-	void InsertCleanupObjectsFence() const;
 	void UpdateStateMachineDeferred(float Delay = 0.0f) const;
-	void GarbageCollectAndUpdateStateMachineDeferred(bool bWaitForCleanupQueue = false) const;
+	void GarbageCollectAndUpdateStateMachineDeferred() const;
 	void UpdateStateMachineImmediate() const;
 
 	void UpdateProgress(float Progress) const;
diff --git a/Engine/Source/Runtime/RenderCore/Private/ShaderCodeLibrary.cpp b/Engine/Source/Runtime/RenderCore/Private/ShaderCodeLibrary.cpp
--- a/Engine/Source/Runtime/RenderCore/Private/ShaderCodeLibrary.cpp
+++ b/Engine/Source/Runtime/RenderCore/Private/ShaderCodeLibrary.cpp
@@ -3307,10 +3307,27 @@ static void FShaderCodeLibraryPluginMountedCallback(IPlugin& Plugin)
 
 static void FShaderCodeLibraryPluginUnmountedCallback(IPlugin& Plugin)
 {
+	class FShaderCodeLibraryCleanup : public FDeferredCleanupInterface
+	{
+	public:
+		FShaderCodeLibraryCleanup(const FString& Name)
+			: Name(Name)
+		{
+		}
+
+		virtual ~FShaderCodeLibraryCleanup()
+		{
+			FShaderCodeLibrary::CloseLibrary(Name);
+		}
+
+	private:
+		FString Name;
+	};
+
 	if (Plugin.CanContainContent() && Plugin.IsEnabled())
 	{
 		// unload any shader libraries that may exist in this plugin
-		FShaderCodeLibrary::CloseLibrary(Plugin.GetName());
+		BeginCleanup(new FShaderCodeLibraryCleanup(Plugin.GetName()));
 	}
 }
 
