#UE Only override the chunkId in the AddChunkOverride GFA

diff --git a/Engine/Plugins/Experimental/GameFeatures/Source/GameFeatures/Private/GameFeatureAction_AddChunkOverride.cpp b/Engine/Plugins/Experimental/GameFeatures/Source/GameFeatures/Private/GameFeatureAction_AddChunkOverride.cpp
--- a/Engine/Plugins/Experimental/GameFeatures/Source/GameFeatures/Private/GameFeatureAction_AddChunkOverride.cpp
+++ b/Engine/Plugins/Experimental/GameFeatures/Source/GameFeatures/Private/GameFeatureAction_AddChunkOverride.cpp
@@ -173,10 +173,7 @@ void UGameFeatureAction_AddChunkOverride::AddChunkIdOverride()
 		UAssetManager& Manager = UAssetManager::Get();
 
 		FPrimaryAssetRules GFDRules;
-		GFDRules.Priority = 0;
 		GFDRules.ChunkId = ChunkId;
-		GFDRules.bApplyRecursively = true;
-		GFDRules.CookRule = EPrimaryAssetCookRule::AlwaysCook;
 		Manager.SetPrimaryAssetRules(GameFeatureData->GetPrimaryAssetId(), GFDRules);
 
 		for (const FPrimaryAssetTypeInfo& AssetTypeInfo : GameFeatureData->GetPrimaryAssetTypesToScan())
@@ -184,10 +181,7 @@ void UGameFeatureAction_AddChunkOverride::AddChunkIdOverride()
 			FPrimaryAssetRulesCustomOverride Override;
 			Override.PrimaryAssetType = FPrimaryAssetType(AssetTypeInfo.PrimaryAssetType);
 			Override.FilterDirectory.Path = FString::Printf(TEXT("/%s"), *PluginName);
-			Override.Rules.Priority = 0;
 			Override.Rules.ChunkId = ChunkId;
-			Override.Rules.bApplyRecursively = true;
-			Override.Rules.CookRule = AssetTypeInfo.Rules.CookRule;
 			Manager.ApplyCustomPrimaryAssetRulesOverride(Override);
 		}
 	}
