moved leak detection and Garbage collection to Plugin Manager
Before unmounting the plugin path, GC is ran, leaked detection for any references objects under the plugin path

https://p4-swarm.epicgames.net/reviews/25611054

diff --git a/Engine/Plugins/Experimental/GameFeatures/Source/GameFeatures/Private/GameFeaturePluginStateMachine.cpp b/Engine/Plugins/Experimental/GameFeatures/Source/GameFeatures/Private/GameFeaturePluginStateMachine.cpp
--- a/Engine/Plugins/Experimental/GameFeatures/Source/GameFeatures/Private/GameFeaturePluginStateMachine.cpp
+++ b/Engine/Plugins/Experimental/GameFeatures/Source/GameFeatures/Private/GameFeaturePluginStateMachine.cpp
@@ -44,81 +44,19 @@ namespace UE::GameFeatures
 
 	static UE::GameFeatures::FResult CanceledResult = MakeError(StateMachineErrorNamespace + TEXT("Canceled"));
 
-
-	static bool MarkPluginAsGarbageOnUnload = true;
-	static FAutoConsoleVariableRef CVarMarkPluginAsGarbageOnUnload(TEXT("GameFeaturePlugin.MarkPluginAsGarbageOnUnload"),
-		MarkPluginAsGarbageOnUnload,
-		TEXT("Mark up the plugin as garbage when the plugin is unloaded (helps track down gc issues)."));
-
-	static bool MarkPluginAsGarbageOnUnregister = false;
-	static FAutoConsoleVariableRef CVarMarkPluginAsGarbageOnUnregister(TEXT("GameFeaturePlugin.MarkPluginAsGarbageOnUnregister"),
-		MarkPluginAsGarbageOnUnregister,
-		TEXT("Mark up the plugin as garbage when the plugin is unloaded (helps track down gc issues)."));
-
 	static int32 ShouldLogMountedFiles = 0;
 	static FAutoConsoleVariableRef CVarShouldLogMountedFiles(TEXT("GameFeaturePlugin.ShouldLogMountedFiles"),
 		ShouldLogMountedFiles,
 		TEXT("Should the newly mounted files be logged."));
 
-	static bool GVerifyUnload = true;
-	static FAutoConsoleVariableRef CVarVerifyPluginUnload(TEXT("GameFeaturePlugin.VerifyUnload"),
-		GVerifyUnload,
-		TEXT("Verify plugin assets are no longer in memory when unloading."),
-		ECVF_Default);
 
 	static FString GVerifyPluginSkipList;
-	static FAutoConsoleVariableRef CVarVerifyPluginSkipList(TEXT("GameFeaturePlugin.VerifyUnload.SkipList"),
+	static FAutoConsoleVariableRef CVarVerifyPluginSkipList(TEXT("PluginManager.VerifyUnload.SkipList"),
 		GVerifyPluginSkipList,
 		TEXT("Comma-separated list of names of plugins for which to skip verification."),
 		ECVF_Default
 		);
 
-	static bool GRenameLeakedPackages = true;
-	static FAutoConsoleVariableRef CVarRenameLeakedPackages(
-		TEXT("GameFeaturePlugin.LeakedAssetTrace.RenameLeakedPackages"),
-		GRenameLeakedPackages,
-		TEXT("Should packages which are leaked after the Game Feature Plugin is unloaded or unmounted."),
-		ECVF_Default
-	);
-
-	static int32 GLeakedAssetTrace_Severity = 2;
-	static FAutoConsoleVariableRef CVarLeakedAssetTrace_Severity(
-	#if UE_BUILD_SHIPPING
-		TEXT("GameFeaturePlugin.LeakedAssetTrace.Severity.Shipping"),
-	#else
-		TEXT("GameFeaturePlugin.LeakedAssetTrace.Severity"),
-	#endif
-		GLeakedAssetTrace_Severity,
-		TEXT("Controls severity of logging when the engine detects that assets from an Game Feature Plugin were leaked during unloading or unmounting.\n")
-		TEXT("0 - all reference tracing and logging is disabled\n")
-		TEXT("1 - logs an error\n")
-		TEXT("2 - ensure\n")
-		TEXT("3 - fatal error\n"),
-		ECVF_Default
-	);
-
-	static int32 GLeakedAssetTrace_TraceMode = (UE_BUILD_SHIPPING ? 0 : 1);
-	static FAutoConsoleVariableRef CVarLeakedAssetTrace_TraceMode(
-	#if UE_BUILD_SHIPPING
-		TEXT("GameFeaturePlugin.LeakedAssetTrace.TraceMode.Shipping"),
-	#else
-		TEXT("GameFeaturePlugin.LeakedASsetTrace.TraceMode"),
-	#endif
-		GLeakedAssetTrace_TraceMode,
-		TEXT("Controls detail level of reference tracing when the engine detects that assets from a Game Feature Plugin were leaked during unloading or unmounting.\n")
-		TEXT("0 - direct references only\n")
-		TEXT("1 - full reference trace"),
-		ECVF_Default
-	);
-
-	static int32 GLeakedAssetTrace_MaxReportCount = 10;
-	static FAutoConsoleVariableRef CVarLeakedAssetTrace_MaxReportCount(
-		TEXT("GameFeaturePlugin.LeakedAssetTrace.MaxReportCount"),
-		GLeakedAssetTrace_MaxReportCount,
-		TEXT("Max number of assets to report when we find leaked assets.\n"),
-		ECVF_Default
-	);
-
 	#define GAME_FEATURE_PLUGIN_STATE_TO_STRING(inEnum, inText) case EGameFeaturePluginState::inEnum: return TEXT(#inEnum);
 	FString ToString(EGameFeaturePluginState InType)
 	{
@@ -176,119 +114,6 @@ namespace UE::GameFeatures
 		return Priority;
 	}
 
-	// Check if any assets from the plugin mount point have leaked, and if so trace them.
-	// Then rename them to allow new copies of them to be loaded. 
-	void HandlePossibleAssetLeaks(const FString& PluginName, UPackage* IgnorePackage = nullptr)
-	{
-		TRACE_CPUPROFILER_EVENT_SCOPE(HandlePossibleAssetLeaks);
-
-		const bool bFindLeakedPackages = GVerifyUnload && (GLeakedAssetTrace_Severity != 0 || GRenameLeakedPackages);
-		if (!bFindLeakedPackages)
-		{
-			return;
-		}
-
-
-		if (UE::GameFeatures::ShouldSkipVerify(PluginName))
-		{
-			return;
-		}
-		
-		TMap<UPackage*, int32> LeakedPackages;
-		TStringBuilder<512> Prefix;
-		Prefix << '/' << PluginName << '/';
-
-		TStringBuilder<NAME_SIZE> NameBuffer;
-		{
-			// If the UObject hash knew about package mount roots, we could avoid this loop
-			TRACE_CPUPROFILER_EVENT_SCOPE(PackageLoop);
-			FPermanentObjectPoolExtents PermanentPool;
-			ForEachObjectOfClass(UPackage::StaticClass(), [&](UObject* Obj)
-			{
-				if (UPackage* Package = CastChecked<UPackage>(Obj, ECastCheckedType::NullAllowed))
-				{
-					if (Package == IgnorePackage)
-					{
-						return;
-					}
-					if (PermanentPool.Contains(Package))
-					{
-						return;
-					}
-					NameBuffer.Reset();
-					Package->GetFName().GetDisplayNameEntry()->AppendNameToString(NameBuffer);
-					if (NameBuffer.ToView().StartsWith(Prefix, ESearchCase::IgnoreCase))
-					{
-						LeakedPackages.Add(Package, GetPackageLeakReportingPriority(Package));
-					}
-				}
-			});
-		}
-
-		if (LeakedPackages.Num() == 0)
-		{
-			return;
-		}
-
-		if (GLeakedAssetTrace_Severity != 0)
-		{
-			EPrintStaleReferencesOptions Options = EPrintStaleReferencesOptions::None;
-			switch (GLeakedAssetTrace_Severity)
-			{
-			case 3:
-				Options = EPrintStaleReferencesOptions::Fatal;
-				break;
-			case 2: 
-				Options = EPrintStaleReferencesOptions::Ensure | EPrintStaleReferencesOptions::Error;
-				break;
-			case 1:
-			default:
-				Options = EPrintStaleReferencesOptions::Error;
-				break;
-			}	
-
-			if (GLeakedAssetTrace_TraceMode == 0)
-			{
-				Options |= EPrintStaleReferencesOptions::Minimal;
-			}
-
-			// Sort even if we don't limit the count, so that high priority leaks appear first 
-			LeakedPackages.ValueSort(TGreater<int32>{});
-			int32 OmittedCount = FMath::Max(0, LeakedPackages.Num() - GLeakedAssetTrace_MaxReportCount);
-			TArray<UPackage*> PackagesToSearchFor;
-			int32 i = 0;
-			for (const TPair<UPackage*, int32>& Pair : LeakedPackages)
-			{
-				if (i++ < GLeakedAssetTrace_MaxReportCount)
-				{
-					PackagesToSearchFor.Add(Pair.Key);
-				}
-			}
-
-			UE_LOG(LogGameFeatures, Display, TEXT("Searching for references to %d leaked packages (%d omitted for speed) from plugin %s"), LeakedPackages.Num(), OmittedCount, *PluginName);
-			FReferenceChainSearch::FindAndPrintStaleReferencesToObjects(MakeArrayView((UObject**)PackagesToSearchFor.GetData(), PackagesToSearchFor.Num()), Options);
-		}
-
-		// Rename the packages that we are streaming out so that we can possibly reload another copy of them
-		for (const TPair<UPackage*, int32>& Pair : LeakedPackages)
-		{
-			UPackage* Package = Pair.Key;
-			UE_LOG(LogGameFeatures, Warning, TEXT("Marking leaking package %s as Garbage"), *Package->GetName());
-			ForEachObjectWithPackage(Package, [](UObject* Object)
-			{
-				Object->MarkAsGarbage();
-				return true;
-			}, false);
-			
-			Package->MarkAsGarbage();
-
-			if (!GIsEditor && !UObject::IsPendingKillEnabled())
-			{
-				UE::Object::RenameLeakedPackage(Package);
-			}
-		}
-	}
-
 #define GAME_FEATURE_PLUGIN_PROTOCOL_PREFIX(inEnum, inString) case EGameFeaturePluginProtocol::inEnum: return inString;
 	const TCHAR* GameFeaturePluginProtocolPrefix(EGameFeaturePluginProtocol Protocol)
 	{
@@ -613,50 +438,6 @@ void FGameFeaturePluginState::UpdateStateMachineDeferred(float Delay /*= 0.0f*/)
 	}), Delay);
 }
 
-void FGameFeaturePluginState::GarbageCollectAndUpdateStateMachineDeferred() const
-{
-	GEngine->ForceGarbageCollection(true); // Tick Delayed
-
-	CleanupDeferredUpdateCallbacks();
-
-	FCoreUObjectDelegates::GetPostGarbageCollect().AddRaw(this, &FGameFeaturePluginState::UpdateStateMachineDeferred, 0.0f);
-}
-
-void FGameFeaturePluginState::MarkPluginAsGarbage(bool bMarkGameFeatureDataAsGarbage)
-{
-	TArray<UPackage*> PackagesToUnload;
-	for (TObjectIterator<UPackage> It; It; ++It)
-	{
-		const FNameBuilder PackageName(It->GetFName());
-		const FStringView PackageMountPointName = FPathViews::GetMountPointNameFromPath(PackageName);
-		if (StateProperties.PluginName == PackageMountPointName)
-		{
-			PackagesToUnload.Add(*It);
-		}
-	}
-
-	for (UPackage* Package : PackagesToUnload)
-	{
-		bool bShouldKeep = false;
-		ForEachObjectWithPackage(Package, [&bShouldKeep, bMarkGameFeatureDataAsGarbage](UObject* Object)
-			{
-				if (Object->IsA(UGameFeatureData::StaticClass()) && !bMarkGameFeatureDataAsGarbage)
-				{
-					bShouldKeep = true;
-					return true;
-				}
-				Object->MarkAsGarbage();
-				return true;
-			}, false);
-		if (!bShouldKeep)
-		{
-			UE_LOG(LogGameFeatures, Log, TEXT("Marking package %s as Garbage"), *Package->GetName());
-			Package->MarkAsGarbage();
-		}
-
-	}
-}
-
 void FGameFeaturePluginState::UpdateStateMachineImmediate() const
 {
 	StateProperties.OnRequestUpdateStateMachine.ExecuteIfBound();
@@ -674,8 +455,6 @@ void FGameFeaturePluginState::CleanupDeferredUpdateCallbacks() const
 		FTSTicker::GetCoreTicker().RemoveTicker(TickHandle);
 		TickHandle.Reset();
 	}
-
-	FCoreUObjectDelegates::GetPostGarbageCollect().RemoveAll(this);
 }
 
 bool FGameFeaturePluginState::ShouldVisitUninstallStateBeforeTerminal() const
@@ -2086,14 +1865,14 @@ struct FGameFeaturePluginState_Unregistering : public FGameFeaturePluginState
 {
 	FGameFeaturePluginState_Unregistering(FGameFeaturePluginStateMachineProperties& InStateProperties) : FGameFeaturePluginState(InStateProperties) {}
 
-	bool bRequestedGC = false;
+	bool bHasUnloaded = false;
 #if WITH_EDITOR
 	bool bRequestedUnloadPluginAssets = false;
 #endif //if WITH_EDITOR
 
 	virtual void BeginState()
 	{
-		bRequestedGC = false;
+		bHasUnloaded = false;
 #if WITH_EDITOR
 		bRequestedUnloadPluginAssets = false;
 #endif //if WITH_EDITOR
@@ -2101,10 +1880,8 @@ struct FGameFeaturePluginState_Unregistering : public FGameFeaturePluginState
 
 	virtual void UpdateState(FGameFeaturePluginStateStatus& StateStatus) override
 	{
-		if (bRequestedGC)
+		if (bHasUnloaded)
 		{
-			UE::GameFeatures::HandlePossibleAssetLeaks(StateProperties.PluginName);
-
 			StateStatus.SetTransition(EGameFeaturePluginState::Unmounting);
 			return;
 		}
@@ -2112,8 +1889,8 @@ struct FGameFeaturePluginState_Unregistering : public FGameFeaturePluginState
 #if WITH_EDITOR
 		if (bRequestedUnloadPluginAssets)
 		{
-			bRequestedGC = true;
-			GarbageCollectAndUpdateStateMachineDeferred();
+			bHasUnloaded = true;
+			UpdateStateMachineDeferred();
 			return;
 		}
 #endif //if WITH_EDITOR
@@ -2140,8 +1917,8 @@ struct FGameFeaturePluginState_Unregistering : public FGameFeaturePluginState
 		{
 			verify(FPluginUtils::UnloadPluginAssets(StateProperties.PluginName));
 
-			bRequestedGC = true;
-			GarbageCollectAndUpdateStateMachineDeferred();
+			bHasUnloaded = true;
+			UpdateStateMachineDeferred();
 		}
 		else
 		{
@@ -2149,14 +1926,8 @@ struct FGameFeaturePluginState_Unregistering : public FGameFeaturePluginState
 			UE::GameFeatures::ScheduleUnloadPluginAssets(StateProperties.PluginName, StateProperties.OnRequestUpdateStateMachine);
 		}
 #else
-		// never mark anything as garbage in the editor, there be dwagons here. 
-		if (UE::GameFeatures::MarkPluginAsGarbageOnUnregister && !UE::GameFeatures::ShouldSkipVerify(StateProperties.PluginName))
-		{
-			MarkPluginAsGarbage(true);
-		}
-
-		bRequestedGC = true;
-		GarbageCollectAndUpdateStateMachineDeferred();
+		bHasUnloaded = true;
+		UpdateStateMachineDeferred();
 #endif
 	}
 };
@@ -2262,38 +2033,29 @@ struct FGameFeaturePluginState_Unloading : public FGameFeaturePluginState
 {
 	FGameFeaturePluginState_Unloading(FGameFeaturePluginStateMachineProperties& InStateProperties) : FGameFeaturePluginState(InStateProperties) {}
 
-	bool bRequestedGC = false;
+	bool bHasUnloaded = false;
 
 	virtual void BeginState() 
 	{
-		bRequestedGC = false;
+		bHasUnloaded = false;
 	}
 
 	virtual void UpdateState(FGameFeaturePluginStateStatus& StateStatus) override
 	{
-		if (bRequestedGC)
+		if (bHasUnloaded)
 		{
-#if !WITH_EDITOR // Disabled in editor since it's likely to report unloaded assets because of standalone packages
-			UE::GameFeatures::HandlePossibleAssetLeaks(StateProperties.PluginName, StateProperties.GameFeatureData ? StateProperties.GameFeatureData->GetOutermost() : nullptr);
-#endif //if !WITH_EDITOR
-
 			StateStatus.SetTransition(EGameFeaturePluginState::Registered);
 			return;
 		}
 
 		UnloadGameFeatureBundles(StateProperties.GameFeatureData);
-#if !WITH_EDITOR // never mark anything as garbage in the editor, there be dwagons here. 
-		if (UE::GameFeatures::MarkPluginAsGarbageOnUnload && !UE::GameFeatures::ShouldSkipVerify(StateProperties.PluginName))
-		{
-			MarkPluginAsGarbage(false);
-		}
-#endif
+
 		if (StateProperties.Destination.MaxState == EGameFeaturePluginState::Registered)
 		{
 			// If we aren't going farther than Registered, GC now
 			// otherwise we will defer until closer to our destination state
-			bRequestedGC = true;
-			GarbageCollectAndUpdateStateMachineDeferred();
+			bHasUnloaded = true;
+			UpdateStateMachineDeferred();
 			return;
 		}
 
@@ -2398,14 +2160,14 @@ struct FGameFeaturePluginState_Deactivating : public FGameFeaturePluginState
 	int32 NumObservedPausers = 0;
 	int32 NumExpectedPausers = 0;
 	bool bInProcessOfDeactivating = false;
-	bool bRequestedGC = false;
+	bool bHasUnloaded = false;
 
 	virtual void BeginState() override
 	{
 		NumObservedPausers = 0;
 		NumExpectedPausers = 0;
 		bInProcessOfDeactivating = false;
-		bRequestedGC = false;
+		bHasUnloaded = false;
 	}
 
 	void OnPauserCompleted(FStringView InPauserTag)
@@ -2424,7 +2186,7 @@ struct FGameFeaturePluginState_Deactivating : public FGameFeaturePluginState
 
 	virtual void UpdateState(FGameFeaturePluginStateStatus& StateStatus) override
 	{
-		if (bRequestedGC)
+		if (bHasUnloaded)
 		{
 			check(NumExpectedPausers == NumObservedPausers);
 			StateStatus.SetTransition(EGameFeaturePluginState::Loaded);
@@ -2460,12 +2222,12 @@ struct FGameFeaturePluginState_Deactivating : public FGameFeaturePluginState
 				UGameFeaturesSubsystem::Get().OnGameFeaturePauseChange(StateProperties.PluginIdentifier.GetFullPluginURL(), StateProperties.PluginName, PauseContext);
 			}
 
-			if (!bRequestedGC && StateProperties.Destination.MaxState == EGameFeaturePluginState::Loaded)
+			if (!bHasUnloaded && StateProperties.Destination.MaxState == EGameFeaturePluginState::Loaded)
 			{
 				// If we aren't going farther than Loaded, GC now
 				// otherwise we will defer until closer to our destination state
-				bRequestedGC = true;
-				GarbageCollectAndUpdateStateMachineDeferred();
+				bHasUnloaded = true;
+				UpdateStateMachineDeferred();
 			}
 			else
 			{
diff --git a/Engine/Plugins/Experimental/GameFeatures/Source/GameFeatures/Private/GameFeaturePluginStateMachine.h b/Engine/Plugins/Experimental/GameFeatures/Source/GameFeatures/Private/GameFeaturePluginStateMachine.h
--- a/Engine/Plugins/Experimental/GameFeatures/Source/GameFeatures/Private/GameFeaturePluginStateMachine.h
+++ b/Engine/Plugins/Experimental/GameFeatures/Source/GameFeatures/Private/GameFeaturePluginStateMachine.h
@@ -353,14 +353,10 @@ struct FGameFeaturePluginState
 	FGameFeaturePluginStateMachineProperties& StateProperties;
 
 	void UpdateStateMachineDeferred(float Delay = 0.0f) const;
-	void GarbageCollectAndUpdateStateMachineDeferred() const;
 	void UpdateStateMachineImmediate() const;
 
 	void UpdateProgress(float Progress) const;
 
-
-	void MarkPluginAsGarbage(bool bMarkGameFeatureDataAsGarbage);
-
 protected:
 	/** Builds an end FResult with some minimal error information with overrides for common types we 
 		need to generate errors from */
diff --git a/Engine/Source/Programs/LowLevelTests/FoundationTests.Target.cs b/Engine/Source/Programs/LowLevelTests/FoundationTests.Target.cs
--- a/Engine/Source/Programs/LowLevelTests/FoundationTests.Target.cs
+++ b/Engine/Source/Programs/LowLevelTests/FoundationTests.Target.cs
@@ -10,6 +10,7 @@ public class FoundationTestsTarget : TestTargetRules
 	{
 		// Collects all tests decorated with #if WITH_LOW_LEVELTESTS from dependencies
 		bWithLowLevelTestsOverride = true;
+		bCompileWithPluginSupport = true;
 
 		bBuildWithEditorOnlyData = Target.Platform.IsInGroup(UnrealPlatformGroup.Desktop)
 			&& (Target.Configuration == UnrealTargetConfiguration.Debug || Target.Configuration == UnrealTargetConfiguration.Development);
diff --git a/Engine/Source/Runtime/CoreUObject/Private/Misc/CoreUObjectPluginManager.cpp b/Engine/Source/Runtime/CoreUObject/Private/Misc/CoreUObjectPluginManager.cpp
new file mode 100644
--- /dev/null
+++ b/Engine/Source/Runtime/CoreUObject/Private/Misc/CoreUObjectPluginManager.cpp
@@ -0,0 +1,201 @@
+// Copyright Epic Games, Inc. All Rights Reserved.
+
+#include "CoreUObjectPluginManager.h"
+#include "Interfaces/IPluginManager.h"
+#include "UObject/UObjectAllocator.h"
+#include "UObject/Package.h"
+#include "Templates/Casts.h"
+#include "UObject/UObjectHash.h"
+#include "Logging/LogMacros.h"
+#include "UObject/ObjectRename.h"
+#include "HAL/IConsoleManager.h"
+#include "Misc/AsciiSet.h"
+#include "String/ParseTokens.h"
+#include "Misc/WildcardString.h"
+#include "UObject/ReferenceChainSearch.h"
+
+namespace UE::CoreUObject::Private
+{
+	
+	static bool GVerifyUnload = true;
+	static FAutoConsoleVariableRef CVarVerifyPluginUnload(TEXT("PluginManager.VerifyUnload"),
+		GVerifyUnload,
+		TEXT("Verify plugin assets are no longer in memory when unloading."),
+		ECVF_Default);
+
+#if WITH_LOW_LEVEL_TESTS
+	bool bEnsureOnLeakedPackages = false;
+#endif
+	static int32 GLeakedAssetTrace_Severity = 2;
+	static FAutoConsoleVariableRef CVarLeakedAssetTrace_Severity(
+#if UE_BUILD_SHIPPING
+		TEXT("PluginManager.LeakedAssetTrace.Severity.Shipping"),
+#else
+		TEXT("PluginManager.LeakedAssetTrace.Severity"),
+#endif
+		GLeakedAssetTrace_Severity,
+		TEXT("Controls severity of logging when the engine detects that assets from an Game Feature Plugin were leaked during unloading or unmounting.\n")
+		TEXT("0 - all reference tracing and logging is disabled\n")
+		TEXT("1 - logs an error\n")
+		TEXT("2 - ensure\n")
+		TEXT("3 - fatal error\n"),
+		ECVF_Default
+	);
+
+	static bool GRenameLeakedPackages = true;
+	static FAutoConsoleVariableRef CVarRenameLeakedPackages(
+		TEXT("PluginManager.LeakedAssetTrace.RenameLeakedPackages"),
+		GRenameLeakedPackages,
+		TEXT("Should packages which are leaked after the Game Feature Plugin is unloaded or unmounted."),
+		ECVF_Default
+	);
+
+	static int32 GLeakedAssetTrace_TraceMode = (UE_BUILD_SHIPPING ? 0 : 1);
+	static FAutoConsoleVariableRef CVarLeakedAssetTrace_TraceMode(
+#if UE_BUILD_SHIPPING
+		TEXT("PluginManager.LeakedAssetTrace.TraceMode.Shipping"),
+#else
+		TEXT("PluginManager.LeakedAssetTrace.TraceMode"),
+#endif
+		GLeakedAssetTrace_TraceMode,
+		TEXT("Controls detail level of reference tracing when the engine detects that assets from a Game Feature Plugin were leaked during unloading or unmounting.\n")
+		TEXT("0 - direct references only\n")
+		TEXT("1 - full reference trace"),
+		ECVF_Default
+	);
+
+	static int32 GLeakedAssetTrace_MaxReportCount = 10;
+	static FAutoConsoleVariableRef CVarLeakedAssetTrace_MaxReportCount(
+		TEXT("PluginManager.LeakedAssetTrace.MaxReportCount"),
+		GLeakedAssetTrace_MaxReportCount,
+		TEXT("Max number of assets to report when we find leaked assets.\n"),
+		ECVF_Default
+	);
+
+	DEFINE_LOG_CATEGORY_STATIC(PluginHandlerLog, Log, All);
+
+	PluginHandler GPluginHandler;
+
+	void PluginHandler::Install()
+	{
+		UE::PluginManager::Private::SetCoreUObjectPluginManager(GPluginHandler);
+	}
+
+	// Check if any assets from the plugin mount point have leaked, and if so trace them.
+	// Then rename them to allow new copies of them to be loaded. 
+	void HandlePossibleAssetLeaks(const FString& PluginName)
+	{
+		TRACE_CPUPROFILER_EVENT_SCOPE(HandlePossibleAssetLeaks);
+
+		const bool bFindLeakedPackages = GVerifyUnload && (GLeakedAssetTrace_Severity != 0 || GRenameLeakedPackages);
+		if (!bFindLeakedPackages)
+		{
+			return;
+		}
+
+		TArray<UPackage*> LeakedPackages;
+		TStringBuilder<512> Prefix;
+		Prefix << '/' << PluginName << '/';
+
+		TStringBuilder<NAME_SIZE> NameBuffer;
+		{
+			// If the UObject hash knew about package mount roots, we could avoid this loop
+			TRACE_CPUPROFILER_EVENT_SCOPE(PackageLoop);
+			FPermanentObjectPoolExtents PermanentPool;
+			ForEachObjectOfClass(UPackage::StaticClass(), [&](UObject* Obj)
+				{
+					if (UPackage* Package = CastChecked<UPackage>(Obj, ECastCheckedType::NullAllowed))
+					{
+						if (PermanentPool.Contains(Package))
+						{
+							return;
+						}
+						NameBuffer.Reset();
+						Package->GetFName().GetDisplayNameEntry()->AppendNameToString(NameBuffer);
+						if (NameBuffer.ToView().StartsWith(Prefix, ESearchCase::IgnoreCase))
+						{
+							LeakedPackages.Add(Package);
+						}
+					}
+				});
+		}
+
+		if (LeakedPackages.Num() == 0)
+		{
+			return;
+		}
+
+#if WITH_LOW_LEVEL_TESTS
+		if (bEnsureOnLeakedPackages)
+		{
+			check(LeakedPackages.Num() == 0);
+		}
+#endif
+
+		if (GLeakedAssetTrace_Severity != 0)
+		{
+			EPrintStaleReferencesOptions Options = EPrintStaleReferencesOptions::None;
+			switch (GLeakedAssetTrace_Severity)
+			{
+			case 3:
+				Options = EPrintStaleReferencesOptions::Fatal;
+				break;
+			case 2:
+				Options = EPrintStaleReferencesOptions::Ensure | EPrintStaleReferencesOptions::Error;
+				break;
+			case 1:
+			default:
+				Options = EPrintStaleReferencesOptions::Error;
+				break;
+			}
+
+			if (GLeakedAssetTrace_TraceMode == 0)
+			{
+				Options |= EPrintStaleReferencesOptions::Minimal;
+			}
+
+			// Sort even if we don't limit the count, so that high priority leaks appear first 
+			int32 OmittedCount = FMath::Max(0, LeakedPackages.Num() - GLeakedAssetTrace_MaxReportCount);
+			TArray<UPackage*> PackagesToSearchFor;
+			int32 i = 0;
+			for (UPackage* Package : LeakedPackages)
+			{
+				if (i++ < GLeakedAssetTrace_MaxReportCount)
+				{
+					PackagesToSearchFor.Add(Package);
+				}
+			}
+
+			UE_LOG(PluginHandlerLog, Display, TEXT("Searching for references to %d leaked packages (%d omitted for speed) from plugin %s"), LeakedPackages.Num(), OmittedCount, *PluginName);
+			FReferenceChainSearch::FindAndPrintStaleReferencesToObjects(MakeArrayView((UObject**)PackagesToSearchFor.GetData(), PackagesToSearchFor.Num()), Options);
+		}
+
+		// Rename the packages that we are streaming out so that we can possibly reload another copy of them
+		for (UPackage* Package : LeakedPackages)
+		{
+			UE_LOG(PluginHandlerLog, Warning, TEXT("Marking leaking package %s as Garbage"), *Package->GetName());
+			ForEachObjectWithPackage(Package, [](UObject* Object)
+				{
+					Object->MarkAsGarbage();
+					return true;
+				}, false);
+
+			Package->MarkAsGarbage();
+
+			if ((!GIsEditor && !UObject::IsPendingKillEnabled())
+#if WITH_LOW_LEVEL_TESTS
+				|| GRenameLeakedPackages
+#endif
+				)
+			{
+				UE::Object::RenameLeakedPackage(Package);
+			}
+		}
+	}
+
+	void PluginHandler::OnPluginUnload(IPlugin& Plugin)
+	{
+		/*CollectGarbage(GARBAGE_COLLECTION_KEEPFLAGS, true);
+		HandlePossibleAssetLeaks(Plugin.GetName());*/
+	}
+}
\ No newline at end of file
diff --git a/Engine/Source/Runtime/CoreUObject/Private/Misc/CoreUObjectPluginManager.h b/Engine/Source/Runtime/CoreUObject/Private/Misc/CoreUObjectPluginManager.h
new file mode 100644
--- /dev/null
+++ b/Engine/Source/Runtime/CoreUObject/Private/Misc/CoreUObjectPluginManager.h
@@ -0,0 +1,15 @@
+// Copyright Epic Games, Inc. All Rights Reserved.
+
+#pragma once
+#include "ICoreUObjectPluginManager.h"
+
+namespace UE::CoreUObject::Private
+{
+	class PluginHandler : public UE::PluginManager::Private::ICoreUObjectPluginManager
+	{
+	public:
+		static void Install();
+
+		virtual void OnPluginUnload(IPlugin& Plugin) override;
+	};
+}
\ No newline at end of file
diff --git a/Engine/Source/Runtime/CoreUObject/Private/Misc/PackageName.cpp b/Engine/Source/Runtime/CoreUObject/Private/Misc/PackageName.cpp
--- a/Engine/Source/Runtime/CoreUObject/Private/Misc/PackageName.cpp
+++ b/Engine/Source/Runtime/CoreUObject/Private/Misc/PackageName.cpp
@@ -25,6 +25,7 @@
 #include "Misc/PackageSegment.h"
 #include "Misc/Paths.h"
 #include "Misc/PathViews.h"
+#include "Misc/CoreUObjectPluginManager.h"
 #include "Misc/ScopeRWLock.h"
 #include "Misc/StringBuilder.h"
 #include "Modules/ModuleManager.h"
@@ -244,8 +245,10 @@ struct FLongPackagePathsSingleton
 	{
 		// Allow the plugin manager to mount new content paths by exposing access through a delegate.  PluginManager is 
 		// a Core class, but content path functionality is added at the CoreUObject level.
-		IPluginManager::Get().SetRegisterMountPointDelegate(IPluginManager::FRegisterMountPointDelegate::CreateStatic(&FPackageName::RegisterMountPoint));
-		IPluginManager::Get().SetUnRegisterMountPointDelegate(IPluginManager::FRegisterMountPointDelegate::CreateStatic(&FPackageName::UnRegisterMountPoint));
+		IPluginManager& PluginManager = IPluginManager::Get();
+		PluginManager.SetRegisterMountPointDelegate(IPluginManager::FRegisterMountPointDelegate::CreateStatic(&FPackageName::RegisterMountPoint));
+		PluginManager.SetUnRegisterMountPointDelegate(IPluginManager::FRegisterMountPointDelegate::CreateStatic(&FPackageName::UnRegisterMountPoint));
+		UE::CoreUObject::Private::PluginHandler::Install();
 	}
 
 	/**
diff --git a/Engine/Source/Runtime/CoreUObject/Tests/PluginManagerTest.cpp b/Engine/Source/Runtime/CoreUObject/Tests/PluginManagerTest.cpp
new file mode 100644
--- /dev/null
+++ b/Engine/Source/Runtime/CoreUObject/Tests/PluginManagerTest.cpp
@@ -0,0 +1,232 @@
+// Copyright Epic Games, Inc. All Rights Reserved.
+
+#if WITH_LOW_LEVEL_TESTS
+#include "Interfaces/IPluginManager.h"
+#include "TestHarness.h"
+#include "HAL/PlatformOutputDevices.h"
+#include "HAL/PlatformFileManager.h"
+#include "GenericPlatform/GenericPlatformFile.h"
+#include "Misc/ScopeExit.h"
+#include "HAL/UnrealMemory.h"
+#include "Templates/Function.h"
+#include "Misc/PackageName.h"
+#include "ObjectPtrTestClass.h"
+#include "UObject/Package.h"
+#include "HAL/IConsoleManager.h"
+#include "UObject/Class.h"
+
+#include "UObject/FastReferenceCollector.h"
+
+
+namespace UE::CoreUObject::Private
+{
+	extern bool bEnsureOnLeakedPackages;
+}
+
+namespace UE::CoreUObject::Private::Tests
+{
+	class MockReadOnlyFileHandle : public IFileHandle
+	{
+	public:
+
+		MockReadOnlyFileHandle(const uint8* Data, int64 Size, int64 Pos)
+			: Data(Data)
+			, Size(Size)
+			, Pos(Pos)
+		{
+		}
+
+		virtual bool Seek(int64 NewPosition)
+		{
+			check(NewPosition < Size);
+			Pos = NewPosition;
+			return true;
+		}
+
+		virtual bool SeekFromEnd(int64 NewPositionRelativeToEnd = 0)
+		{
+			Pos = Size - NewPositionRelativeToEnd;
+			return true;
+		}
+
+		virtual bool Read(uint8* Destination, int64 BytesToRead)
+		{
+			FMemory::Memcpy(Destination, Data + Pos, BytesToRead);
+			return true;
+		}
+
+		virtual int64 Tell() { return Pos; }
+		virtual bool Write(const uint8* Source, int64 BytesToWrite) { return false; }
+		virtual bool Flush(const bool bFullFlush = false){ return false; }
+		virtual bool Truncate(int64 NewSize) { return false; }
+
+		const uint8* Data;
+		int64 Size;
+		int64 Pos;
+	};
+
+	class MockPlatformFile : public IPlatformFile
+	{
+	public:
+
+		MockPlatformFile(IPlatformFile& File)
+			: PlatformFile(File)
+		{
+		}
+
+		virtual bool Initialize(IPlatformFile* Inner, const TCHAR* CmdLine) override { return false; }
+		virtual IPlatformFile* GetLowerLevel() override { return &PlatformFile; }
+		virtual void SetLowerLevel(IPlatformFile* NewLowerLevel) override { }
+		virtual const TCHAR* GetName() const override { return TEXT("MockPlatformFile") ; }
+		virtual bool FileExists(const TCHAR* Filename) override { return false; }
+		virtual int64 FileSize(const TCHAR* Filename) override { return 0; }
+		virtual bool DeleteFile(const TCHAR* Filename) override { return false; }
+		virtual bool IsReadOnly(const TCHAR* Filename) override { return false; }
+		virtual bool MoveFile(const TCHAR* To, const TCHAR* From) override { return false; }
+		virtual bool SetReadOnly(const TCHAR* Filename, bool bNewReadOnlyValue) override { return false; }
+		virtual FDateTime GetTimeStamp(const TCHAR* Filename) override { return  FDateTime(); }
+		virtual void SetTimeStamp(const TCHAR* Filename, FDateTime DateTime) override { }
+		virtual FDateTime GetAccessTimeStamp(const TCHAR* Filename) override { return FDateTime(); }
+		virtual FString GetFilenameOnDisk(const TCHAR* Filename) override { return TEXT(""); }
+		virtual IFileHandle* OpenWrite(const TCHAR* Filename, bool bAppend = false, bool bAllowRead = false) override { return nullptr; }
+		virtual bool DirectoryExists(const TCHAR* Directory) override { return false; }
+		virtual bool CreateDirectory(const TCHAR* Directory) override { return false; }
+		virtual bool DeleteDirectory(const TCHAR* Directory) override { return false; }
+		virtual FFileStatData GetStatData(const TCHAR* FilenameOrDirectory) override { return FFileStatData(); }
+		virtual bool IterateDirectory(const TCHAR* Directory, FDirectoryVisitor& Visitor) { return false; }
+		virtual bool IterateDirectoryStat(const TCHAR* Directory, FDirectoryStatVisitor& Visitor) { return false; }
+
+		virtual IFileHandle* OpenRead(const TCHAR* Filename, bool bAllowWrite = false) override 
+		{
+			TFunction<IFileHandle* ()>* Handle = Files.Find(Filename);
+			if (Handle)
+			{
+				return (*Handle)();
+			}
+			return nullptr;
+		}
+
+		IPlatformFile& PlatformFile;
+		TMap<FString, TFunction<IFileHandle* ()>> Files;
+	};
+
+	class FProcessor : public FSimpleReferenceProcessorBase
+	{
+	public:
+		void HandleTokenStreamObjectReference(FGCArrayStruct& ObjectsToSerializeStruct, UObject* ReferencingObject, UObject*& Object, UE::GC::FTokenId, EGCTokenType, bool);
+	};
+
+	void FProcessor::HandleTokenStreamObjectReference(FGCArrayStruct& ObjectsToSerializeStruct, UObject* ReferencingObject, UObject*& Object, UE::GC::FTokenId, EGCTokenType, bool)
+	{
+		printf("");
+	}
+
+	TEST_CASE("UE::CoreUObject::PluginHandler::LeakDetection")
+	{
+		const ANSICHAR* PluginText = R"(
+		{ 
+			"FileVersion": 3,
+			"ExplicitlyLoaded" : true,
+			"CanContainContent" : true
+		 })";
+		MockPlatformFile MockFile(FPlatformFileManager::Get().GetPlatformFile());
+		MockFile.Files.Add(TEXT("MyTestPlugin.uplugin"), [PluginText]()
+			{
+				return new MockReadOnlyFileHandle((const uint8*)PluginText, FCStringAnsi::Strlen(PluginText), 0u );
+			});
+
+		FPlatformFileManager::Get().SetPlatformFile(MockFile);
+		IConsoleVariable* RenamedLeakedPackages = IConsoleManager::Get().FindConsoleVariable(TEXT("PluginManager.LeakedAssetTrace.RenameLeakedPackages"));
+		IConsoleVariable* LeakSeverity = IConsoleManager::Get().FindConsoleVariable(TEXT("PluginManager.LeakedAssetTrace.Severity"));
+		CHECK(RenamedLeakedPackages);
+		CHECK(LeakSeverity);
+		bool OldRenameValue = RenamedLeakedPackages->GetBool();
+		int OldSeverityValue = RenamedLeakedPackages->GetInt();
+		RenamedLeakedPackages->Set(true);
+		LeakSeverity->Set(0); //disabled leak logging
+		ON_SCOPE_EXIT
+		{
+			RenamedLeakedPackages->Set(OldRenameValue);
+			LeakSeverity->Set(OldSeverityValue);
+			FPlatformFileManager::Get().RemovePlatformFile(&MockFile);
+		};
+		IPluginManager& PluginManager = IPluginManager::Get();
+		CHECK(PluginManager.AddToPluginsList("MyTestPlugin.uplugin"));
+		TSharedPtr<IPlugin> Plugin = PluginManager.FindPlugin("MyTestPlugin");
+		CHECK(Plugin);
+		CHECK(PluginManager.MountExplicitlyLoadedPlugin(TEXT("MyTestPlugin")));
+		CHECK(FPackageName::MountPointExists(TEXT("/MyTestPlugin")));
+
+		UPackage* PluginPackage = NewObject<UPackage>(nullptr, TEXT("/MyTestPlugin/MyPackage"), RF_Public);
+		
+		FName PluginPackageName = PluginPackage->GetFName();
+		UObjectPtrTestClass* Obj = NewObject<UObjectPtrTestClass>(PluginPackage, TEXT("Obj"), RF_Public);
+
+		UPackage* OtherPackage = NewObject<UPackage>(nullptr, TEXT("/Test/TestPackage"), RF_Public);
+		UObjectPtrTestClassWithRef* RefObj = NewObject<UObjectPtrTestClassWithRef>(OtherPackage, TEXT("RefObj"), RF_Public);
+
+		OtherPackage->AddToRoot(); 
+		RefObj->AddToRoot();
+		//don't root PluginPackage as rooted packages can not be renamed
+		ON_SCOPE_EXIT
+		{
+			OtherPackage->RemoveFromRoot();
+			RefObj->RemoveFromRoot();
+		};
+		RefObj->ObjectPtr = Obj; //add a reference
+
+		FText Reason;
+		CHECK(PluginManager.UnmountExplicitlyLoadedPlugin(TEXT("MyTestPlugin"), &Reason));
+		CHECK(!FPackageName::MountPointExists(TEXT("/MyTestPlugin")));
+		CHECK(PluginPackageName != PluginPackage->GetFName()); //package should have been renamed.
+		CHECK(PluginPackage->HasAnyInternalFlags(EInternalObjectFlags::MirroredFlags));
+		CHECK(Obj->HasAnyInternalFlags(EInternalObjectFlags::MirroredFlags));
+
+		UPackage* PluginPackage2 = NewObject<UPackage>(nullptr, TEXT("/MyTestPlugin/MyPackage"), RF_Public);
+		CHECK(PluginPackage != PluginPackage2);
+	}
+
+	TEST_CASE("UE::CoreUObject::PluginHandler::GarbageCollection")
+	{
+		const ANSICHAR* PluginText = R"(
+		{ 
+			"FileVersion": 3,
+			"ExplicitlyLoaded" : true,
+			"CanContainContent" : true
+		 })";
+		MockPlatformFile MockFile(FPlatformFileManager::Get().GetPlatformFile());
+		MockFile.Files.Add(TEXT("MyTestPlugin.uplugin"), [PluginText]()
+			{
+				return new MockReadOnlyFileHandle((const uint8*)PluginText, FCStringAnsi::Strlen(PluginText), 0u);
+			});
+
+		FPlatformFileManager::Get().SetPlatformFile(MockFile);
+		UE::CoreUObject::Private::bEnsureOnLeakedPackages = true;
+		ON_SCOPE_EXIT
+		{
+			UE::CoreUObject::Private::bEnsureOnLeakedPackages = false;
+			FPlatformFileManager::Get().RemovePlatformFile(&MockFile);
+		};
+		IPluginManager& PluginManager = IPluginManager::Get();
+		CHECK(PluginManager.AddToPluginsList("MyTestPlugin.uplugin"));
+		TSharedPtr<IPlugin> Plugin = PluginManager.FindPlugin("MyTestPlugin");
+		CHECK(Plugin);
+		CHECK(PluginManager.MountExplicitlyLoadedPlugin(TEXT("MyTestPlugin")));
+		CHECK(FPackageName::MountPointExists(TEXT("/MyTestPlugin")));
+
+		UPackage* PluginPackage = NewObject<UPackage>(nullptr, TEXT("/MyTestPlugin/MyPackage"), RF_Public);
+
+		FName PluginPackageName = PluginPackage->GetFName();
+		UObjectPtrTestClass* Obj = NewObject<UObjectPtrTestClass>(PluginPackage, TEXT("Obj"), RF_Public);
+		UObjectPtrTestClassWithRef* RefObj = NewObject<UObjectPtrTestClassWithRef>(PluginPackage, TEXT("RefObj"), RF_Public);
+
+		FTopLevelAssetPath AssetPath(Obj);
+		CHECK(FindObject<UObjectPtrTestClass>(AssetPath));
+
+		FText Reason;
+		CHECK(PluginManager.UnmountExplicitlyLoadedPlugin(TEXT("MyTestPlugin"), &Reason));
+		CHECK(!FPackageName::MountPointExists(TEXT("/MyTestPlugin")));
+		CHECK(!FindObject<UObjectPtrTestClass>(AssetPath));
+	}
+}
+#endif
diff --git a/Engine/Source/Runtime/Projects/Internal/ICoreUObjectPluginManager.h b/Engine/Source/Runtime/Projects/Internal/ICoreUObjectPluginManager.h
new file mode 100644
--- /dev/null
+++ b/Engine/Source/Runtime/Projects/Internal/ICoreUObjectPluginManager.h
@@ -0,0 +1,19 @@
+// Copyright Epic Games, Inc. All Rights Reserved.
+
+#pragma once
+#include "HAL/Platform.h"
+
+class IPlugin;
+
+namespace UE::PluginManager::Private
+{
+	/// @brief Interface to get around dependency problems between Projects and CoreUObject
+	/// interface allows CoreUObject to run login in PluginManager that Projects can't run
+	class ICoreUObjectPluginManager
+	{
+	public:
+		virtual void OnPluginUnload(IPlugin& Plugin) = 0;
+	};
+
+	PROJECTS_API void SetCoreUObjectPluginManager(ICoreUObjectPluginManager& Handler);
+}
\ No newline at end of file
diff --git a/Engine/Source/Runtime/Projects/Private/PluginManager.cpp b/Engine/Source/Runtime/Projects/Private/PluginManager.cpp
--- a/Engine/Source/Runtime/Projects/Private/PluginManager.cpp
+++ b/Engine/Source/Runtime/Projects/Private/PluginManager.cpp
@@ -1,6 +1,7 @@
 // Copyright Epic Games, Inc. All Rights Reserved.
 
 #include "PluginManager.h"
+#include "ICoreUObjectPluginManager.h"
 #include "GenericPlatform/GenericPlatformFile.h"
 #include "HAL/PlatformFileManager.h"
 #include "HAL/FileManager.h"
@@ -42,6 +43,14 @@ DEFINE_LOG_CATEGORY_STATIC( LogPluginManager, Log, All );
 
 namespace UE::PluginManager::Private
 {
+
+ICoreUObjectPluginManager* CoreUObjectPluginHandler = nullptr;
+
+void SetCoreUObjectPluginManager(ICoreUObjectPluginManager& Handler)
+{
+	CoreUObjectPluginHandler = &Handler;
+}
+	
 TArray<FString> GetPluginPathsByEnv( const TCHAR* EnvVariable )
 {
 	TArray<FString> AdditionalPaths;
@@ -2600,7 +2609,7 @@ TArray<TSharedRef<IPlugin>> FPluginManager::GetEnabledPluginsWithVerse() const
 	{
 		const TSharedRef<FPlugin>& PluginRef = DiscoveredPluginMapUtils::ResolvePluginFromMapVal(PluginPair.Value);
 		const FPlugin& Plugin = *PluginRef;
-		if (Plugin.IsEnabled() && Plugin.CanContainVerse())
+		if (Plugin.IsEnabled() && Plugin.CanContainVerse() && Plugin.IsMounted())
 		{
 			Plugins.Add(PluginRef);
 		}
@@ -2840,12 +2849,6 @@ void FPluginManager::MountPluginFromExternalSource(const TSharedRef<FPlugin>& Pl
 	}
 }
 
-bool FPluginManager::UnmountExplicitlyLoadedPlugin(const FString& PluginName, FText* OutReason)
-{
-	TSharedPtr<FPlugin> Plugin = FindPluginInstance(PluginName);
-	return UnmountPluginFromExternalSource(Plugin, OutReason);
-}
-
 bool FPluginManager::GetPluginDependencies(const FString& PluginName, TArray<FPluginReferenceDescriptor>& PluginDependencies)
 {
 	const TSharedPtr<FPlugin> Plugin = FindPluginInstance(PluginName);
@@ -2879,8 +2882,9 @@ bool FPluginManager::GetPluginDependencies_FromDescriptor(const FPluginReference
 	return false;
 }
 
-bool FPluginManager::UnmountPluginFromExternalSource(const TSharedPtr<FPlugin>& Plugin, FText* OutReason)
+bool FPluginManager::UnmountExplicitlyLoadedPlugin(const FString& PluginName, FText* OutReason)
 {
+	TSharedPtr<FPlugin> Plugin = FindPluginInstance(PluginName);
 	TRACE_CPUPROFILER_EVENT_SCOPE(UnmountPluginFromExternalSource);
 	if (!Plugin.IsValid() || Plugin->bEnabled == false)
 	{
@@ -2928,12 +2932,17 @@ bool FPluginManager::UnmountPluginFromExternalSource(const TSharedPtr<FPlugin>&
 			}
 		}
 
-		UnRegisterMountPointDelegate.Execute(Plugin->GetMountedAssetPath(), Plugin->GetContentDir());
-
 		if (PluginUnmountedEvent.IsBound())
 		{
 			PluginUnmountedEvent.Broadcast(*Plugin);
 		}
+
+		if (UE::PluginManager::Private::CoreUObjectPluginHandler)
+		{
+			UE::PluginManager::Private::CoreUObjectPluginHandler->OnPluginUnload(*Plugin);
+		}
+
+		UnRegisterMountPointDelegate.Execute(Plugin->GetMountedAssetPath(), Plugin->GetContentDir());
 	}
 
 	Plugin->SetIsMounted(false);
diff --git a/Engine/Source/Runtime/Projects/Private/PluginManager.h b/Engine/Source/Runtime/Projects/Private/PluginManager.h
--- a/Engine/Source/Runtime/Projects/Private/PluginManager.h
+++ b/Engine/Source/Runtime/Projects/Private/PluginManager.h
@@ -265,9 +265,6 @@ private:
 	/** Mounts a plugin that was requested to be mounted from external code (either by MountNewlyCreatedPlugin or MountExplicitlyLoadedPlugin) */
 	void MountPluginFromExternalSource(const TSharedRef<FPlugin>& Plugin);
 
-	/** Unmounts a plugin that was requested to be unmounted from external code (by UnmountExplicitlyLoadedPlugin) */
-	bool UnmountPluginFromExternalSource(const TSharedPtr<FPlugin>& Plugin, FText* OutReason);
-
 #if WITH_EDITOR
 	void AddToModuleNameToPluginMap(const TSharedRef<FPlugin>& Plugin);
 	void RemoveFromModuleNameToPluginMap(const TSharedRef<FPlugin>& Plugin);
