IsValidGameFeaturePlugin
- Use IsParentPathOf instead of StartsWith. StartsWith will not handle missing directory seperators.

diff --git a/Engine/Plugins/Experimental/GameFeatures/Source/GameFeatures/Private/GameFeaturesSubsystemSettings.cpp b/Engine/Plugins/Experimental/GameFeatures/Source/GameFeatures/Private/GameFeaturesSubsystemSettings.cpp
--- a/Engine/Plugins/Experimental/GameFeatures/Source/GameFeatures/Private/GameFeaturesSubsystemSettings.cpp
+++ b/Engine/Plugins/Experimental/GameFeatures/Source/GameFeatures/Private/GameFeaturesSubsystemSettings.cpp
@@ -2,6 +2,7 @@
 
 #include "GameFeaturesSubsystemSettings.h"
 #include "Misc/Paths.h"
+#include "Misc/PathViews.h"
 
 #include UE_INLINE_GENERATED_CPP_BY_NAME(GameFeaturesSubsystemSettings)
 
@@ -40,7 +41,7 @@ bool UGameFeaturesSubsystemSettings::IsValidGameFeaturePlugin(const FString& Plu
 	// Check to see if the filename is rooted in a game feature plugin folder
 	for (const FString& BuiltInFolder : BuiltInGameFeaturePluginsFolders)
 	{
-		if (PluginDescriptorFilename.StartsWith(BuiltInFolder))
+		if (FPathViews::IsParentPathOf(BuiltInFolder, PluginDescriptorFilename))
 		{
 			return true;
 		}
