Support for dynamically creating a transient UGameFeatureData when bCanContainContent is false. This is support work for a "root" plugin concept. Fixing up problematic code locations that manually attempt to get the plugin name from the UGameFeatureData object via UAssetManager::GetContentRootPathFromPackageName. These now need to use UGameFeatureData::GetPluginName instead otherwise the package path will not resolve to a content root matching the plugin name for a transient GFD.

#rnx

diff --git a/Engine/Plugins/Experimental/GameFeatures/Source/GameFeatures/Private/GameFeatureData.cpp b/Engine/Plugins/Experimental/GameFeatures/Source/GameFeatures/Private/GameFeatureData.cpp
--- a/Engine/Plugins/Experimental/GameFeatures/Source/GameFeatures/Private/GameFeatureData.cpp
+++ b/Engine/Plugins/Experimental/GameFeatures/Source/GameFeatures/Private/GameFeatureData.cpp
@@ -280,16 +280,24 @@ void UGameFeatureData::GetPluginName(const UGameFeatureData* GFD, FString& Plugi
 {
 	if (GFD)
 	{
-		const FString GameFeaturePath = GFD->GetOutermost()->GetName();
-		if (ensureMsgf(UAssetManager::GetContentRootPathFromPackageName(GameFeaturePath, PluginName), TEXT("Must be a valid package path with a root. GameFeaturePath: %s"), *GameFeaturePath))
+		const bool bIsTransient = (GFD->GetFlags() & RF_Transient) != 0;
+		if (bIsTransient)
 		{
-			// Trim the leading and trailing slashes
-			PluginName = PluginName.LeftChop(1).RightChop(1);
+			PluginName = GFD->GetName();
 		}
 		else
 		{
-			// Not a great fallback but better than nothing. Make sure this asset is in the right folder so we can get the plugin name.
-			PluginName = GFD->GetName();
+			const FString GameFeaturePath = GFD->GetOutermost()->GetName();
+			if (ensureMsgf(UAssetManager::GetContentRootPathFromPackageName(GameFeaturePath, PluginName), TEXT("Must be a valid package path with a root. GameFeaturePath: %s"), *GameFeaturePath))
+			{
+				// Trim the leading and trailing slashes
+				PluginName = PluginName.LeftChop(1).RightChop(1);
+			}
+			else
+			{
+				// Not a great fallback but better than nothing. Make sure this asset is in the right folder so we can get the plugin name.
+				PluginName = GFD->GetName();
+			}
 		}
 	}
 }
diff --git a/Engine/Plugins/Experimental/GameFeatures/Source/GameFeatures/Private/GameFeaturePluginStateMachine.cpp b/Engine/Plugins/Experimental/GameFeatures/Source/GameFeatures/Private/GameFeaturePluginStateMachine.cpp
--- a/Engine/Plugins/Experimental/GameFeatures/Source/GameFeatures/Private/GameFeaturePluginStateMachine.cpp
+++ b/Engine/Plugins/Experimental/GameFeatures/Source/GameFeatures/Private/GameFeaturePluginStateMachine.cpp
@@ -1984,10 +1984,21 @@ struct FGameFeaturePluginState_Registering : public FGameFeaturePluginState
 			return nullptr;
 		};
 
-		StateProperties.GameFeatureData = LoadGameFeatureData(PreferredGameFeatureDataPath);
-		if (!StateProperties.GameFeatureData)
+		TSharedPtr<IPlugin> Plugin = IPluginManager::Get().FindPlugin(StateProperties.PluginName);
+		ensure(Plugin.IsValid());
+		
+		// If the plugin contains content then load the GameFeatureData otherwise procedurally create one that is transient.
+		if (Plugin->GetDescriptor().bCanContainContent)
+		{
+			StateProperties.GameFeatureData = LoadGameFeatureData(PreferredGameFeatureDataPath);
+			if (!StateProperties.GameFeatureData)
+			{
+				StateProperties.GameFeatureData = LoadGameFeatureData(BackupGameFeatureDataPath);
+			}
+		}
+		else
 		{
-			StateProperties.GameFeatureData = LoadGameFeatureData(BackupGameFeatureDataPath);
+			StateProperties.GameFeatureData = NewObject<UGameFeatureData>(GetTransientPackage(), FName(*StateProperties.PluginName), RF_Transient);
 		}
 
 		if (StateProperties.GameFeatureData)
