GFP URL refactor
- Remove any runtime behavior flags from GFP URLs, instead pass them via a seperate FGameFeatureProtocolOptions struct
- This ensures that GFP URLs contain only information necessary to resolve their data
- Dependency state machines now inherit FGameFeatureProtocolOptions from dependent state machines if they haven't already been created before they are required by the dependent machine
- LoadGameFeaturePlugin may now update runtime options if needed

diff --git a/Engine/Plugins/Experimental/GameFeatures/Source/GameFeatures/Private/GameFeaturePluginStateMachine.cpp b/Engine/Plugins/Experimental/GameFeatures/Source/GameFeatures/Private/GameFeaturePluginStateMachine.cpp
--- a/Engine/Plugins/Experimental/GameFeatures/Source/GameFeatures/Private/GameFeaturePluginStateMachine.cpp
+++ b/Engine/Plugins/Experimental/GameFeatures/Source/GameFeatures/Private/GameFeaturePluginStateMachine.cpp
@@ -38,6 +38,8 @@
 #include "Misc/App.h"
 #endif //if WITH_EDITOR
 
+const uint32 FInstallBundlePluginProtocolMetaData::FDefaultValues::CurrentVersionNum = 2;
+
 namespace UE::GameFeatures
 {
 	static const FString StateMachineErrorNamespace(TEXT("GameFeaturePlugin.StateMachine."));
@@ -260,13 +262,7 @@ namespace UE::GameFeatures
 //Used to aid in parsing / creating the URL for InstallBundle protocol GameFeaturePlugins
 #define INSTALL_BUNDLE_PROTOCOL_URL_OPTIONS_LIST(XOPTION)					\
 	XOPTION(Bundles,					TEXT("Bundles"))					\
-	XOPTION(Version,					TEXT("V"))							\
-	XOPTION(Flags,						TEXT("Flags"))						\
-	XOPTION(ReleaseFlags,				TEXT("ReleaseFlags"))				\
-	XOPTION(UninstallBeforeTerminate,	TEXT("UninstallBeforeTerminate"))	\
-	XOPTION(UserPauseDownload,			TEXT("Paused"))						\
-	XOPTION(AllowIniLoading,			TEXT("AllowIni"))					\
-	XOPTION(DoNotDownload,				TEXT("DoNotDownload"))
+	XOPTION(Version,					TEXT("V"))							
 
 #define INSTALL_BUNDLE_PROTOCOL_URL_OPTIONS_ENUM(inEnum, inString) inEnum,
 enum class EGameFeatureInstallBundleProtocolOptions : uint8
@@ -376,28 +372,15 @@ FGameFeaturePluginState::~FGameFeaturePluginState()
 	CleanupDeferredUpdateCallbacks();
 }
 
-bool FGameFeaturePluginState::TryUpdateURLData(const FString& NewPluginURL)
+UE::GameFeatures::FResult FGameFeaturePluginState::TryUpdateProtocolOptions(const FGameFeatureProtocolOptions& NewOptions)
 {
-	//Early out without any update if we already fully match the new URL
-	if (StateProperties.PluginIdentifier.ExactMatchesURL(NewPluginURL))
-	{
-		return false;
-	}
-
-	//Create copied StateMachineProperties data and try and reset it by parsing the new data
-	FGameFeaturePluginIdentifier NewIdentifier = FGameFeaturePluginIdentifier(NewPluginURL);
-	FGameFeaturePluginStateMachineProperties NewProperties = StateProperties;
-	NewProperties.PluginIdentifier = NewIdentifier;
+	UE::GameFeatures::FResult Result = StateProperties.ValidateProtocolOptionsUpdate(NewOptions);
 
-	if (!NewProperties.ParseURL() || !NewProperties.ValidateURLUpdate(StateProperties))
+	if (!Result.HasError())
 	{
-		return false;
+		StateProperties.ProtocolOptions = NewOptions;
 	}
-
-	//Parse was successful so transfer to the newly parsed properties
-	StateProperties = MoveTemp(NewProperties);
-
-	return true;
+	return Result;
 }
 
 FDestinationGameFeaturePluginState* FGameFeaturePluginState::AsDestinationState()
@@ -464,7 +447,7 @@ bool FGameFeaturePluginState::ShouldVisitUninstallStateBeforeTerminal() const
 		case (EGameFeaturePluginProtocol::InstallBundle):
 		{
 			//InstallBundleProtocol's have a MetaData that controlls if they uninstall currently
-			return StateProperties.ProtocolMetadata.GetSubtype<FInstallBundlePluginProtocolMetaData>().bUninstallBeforeTerminate;
+			return StateProperties.ProtocolOptions.GetSubtype<FInstallBundlePluginProtocolOptions>().bUninstallBeforeTerminate;
 		}
 
 		//Default behavior is to just Terminate
@@ -483,7 +466,7 @@ bool FGameFeaturePluginState::AllowIniLoading() const
 	{
 		//InstallBundleProtocol's have a MetaData that controlls if INI loading is allowed
 		//The protocol default is not to allow INI loading since the source is likely untrusted
-		return StateProperties.ProtocolMetadata.GetSubtype<FInstallBundlePluginProtocolMetaData>().bAllowIniLoading;
+		return StateProperties.ProtocolOptions.GetSubtype<FInstallBundlePluginProtocolOptions>().bAllowIniLoading;
 	}
 
 	//Default behavior is to allow INI loading
@@ -516,10 +499,10 @@ struct FGameFeaturePluginState_Terminal : public FDestinationGameFeaturePluginSt
 
 	bool bEnteredTerminalState = false;
 
-	virtual bool TryUpdateURLData(const FString& NewPluginURL) override
+	virtual UE::GameFeatures::FResult TryUpdateProtocolOptions(const FGameFeatureProtocolOptions& NewOptions) override
 	{
-		//Should never update our URL during Terminal
-		return false;
+		//Should never update our options during Terminal
+		return GetErrorResult(TEXT("ProtocolOptions."), TEXT("Terminal"));
 	}
 
 	virtual void BeginState() override
@@ -879,7 +862,7 @@ struct FBaseDataReleaseGameFeaturePluginState : public FGameFeaturePluginState
 	/** Determine what kind of release request flags we submit */
 	virtual EInstallBundleReleaseRequestFlags GetReleaseRequestFlags() const
 	{
-		return StateProperties.ProtocolMetadata.GetSubtype<FInstallBundlePluginProtocolMetaData>().ReleaseInstallBundleFlags;
+		return StateProperties.ProtocolOptions.GetSubtype<FInstallBundlePluginProtocolOptions>().ReleaseInstallBundleFlags;
 	}
 
 	/** Determines what state you transition to in the event of a success or failure to release content */
@@ -927,12 +910,14 @@ struct FGameFeaturePluginState_Uninstalling : public FBaseDataReleaseGameFeature
 		return (BaseFlags | EInstallBundleReleaseRequestFlags::RemoveFilesIfPossible);
 	}
 
-	virtual bool TryUpdateURLData(const FString& NewPluginURL) override
+	virtual UE::GameFeatures::FResult TryUpdateProtocolOptions(const FGameFeatureProtocolOptions& NewOptions) override
 	{
 		//Use base functionality to update our metadata
-		if (!FGameFeaturePluginState::TryUpdateURLData(NewPluginURL))
+		UE::GameFeatures::FResult LocalResult = FGameFeaturePluginState::TryUpdateProtocolOptions(NewOptions);
+
+		if (LocalResult.HasError())
 		{
-			return false;
+			return LocalResult;
 		}
 
 		//If we are no longer uninstalling before terminate, just exit now as a success immediately
@@ -944,13 +929,13 @@ struct FGameFeaturePluginState_Uninstalling : public FBaseDataReleaseGameFeature
 			bWasDeleted = true;
 			UpdateStateMachineImmediate();
 
-			return true;
+			return LocalResult;
 		}		
 		
 		//Restart our remove request to handle other changes
 		BeginRemoveRequest();
 
-		return true;
+		return LocalResult;
 	}
 };
 
@@ -1050,17 +1035,19 @@ struct FGameFeaturePluginState_Downloading : public FGameFeaturePluginState
 			return;
 		}
 
+		const FInstallBundlePluginProtocolOptions& Options = StateProperties.ProtocolOptions.GetSubtype<FInstallBundlePluginProtocolOptions>();
+
 		// not up to date, check to see if we allow downloading
-		if (Metadata.bDoNotDownload)
+		if (Options.bDoNotDownload)
 		{
 			Result = GetErrorResult(TEXT("GFPStateMachine.DownloadNotAllowed"));
 			UpdateStateMachineImmediate();
 			return;
 		}
 
-		//Pull our InstallFlags from the Metadata, but also make sure SkipMount is set as there is a separate mounting step that will re-request this
+		//Pull our InstallFlags from the Options, but also make sure SkipMount is set as there is a separate mounting step that will re-request this
 		//without SkipMount and then mount the data, this allows us to pre-download data without mounting it
-		EInstallBundleRequestFlags InstallFlags = Metadata.InstallBundleFlags;
+		EInstallBundleRequestFlags InstallFlags = Options.InstallBundleFlags;
 		InstallFlags |= EInstallBundleRequestFlags::SkipMount;
 
 		TValueOrError<FInstallBundleRequestInfo, EInstallBundleResult> MaybeRequestInfo = BundleManager->RequestUpdateContent(InstallBundles, InstallFlags);
@@ -1104,7 +1091,7 @@ struct FGameFeaturePluginState_Downloading : public FGameFeaturePluginState
 
 			//If this setting is flipped then we should immediately request to pause downloads.
 			//We still generate the downloads so that we have an accurate PendingBundleDownloads list
-			if (Metadata.bUserPauseDownload)
+			if (Options.bUserPauseDownload)
 			{
 				ChangePauseState(true);
 			}
@@ -1242,28 +1229,34 @@ struct FGameFeaturePluginState_Downloading : public FGameFeaturePluginState
 		}
 	}
 
-	virtual bool TryUpdateURLData(const FString& NewPluginURL) override
+	virtual UE::GameFeatures::FResult TryUpdateProtocolOptions(const FGameFeatureProtocolOptions& NewOptions) override
 	{
-		//if we don't have any in-progress downloads the default behavior is all we need
-		if (PendingBundleDownloads.Num() == 0)
+		//Need to update our BundleFlags for any bundles we are downloading
+		EInstallBundleRequestFlags OldRequestFlags;
+		bool OldUserPausedFlag;
 		{
-			return FGameFeaturePluginState::TryUpdateURLData(NewPluginURL);
+			const FInstallBundlePluginProtocolOptions& OldOptions = StateProperties.ProtocolOptions.GetSubtype<FInstallBundlePluginProtocolOptions>();
+			OldRequestFlags = OldOptions.InstallBundleFlags;
+			OldUserPausedFlag = OldOptions.bUserPauseDownload;
 		}
-		
-		//Need to update our BundleFlags for any bundles we are downloading
-		FInstallBundlePluginProtocolMetaData& OldMetaData = StateProperties.ProtocolMetadata.GetSubtype<FInstallBundlePluginProtocolMetaData>();
-		EInstallBundleRequestFlags OldRequestFlags = OldMetaData.InstallBundleFlags;
-		const bool OldUserPausedFlag = OldMetaData.bUserPauseDownload;
 
-		if (!FGameFeaturePluginState::TryUpdateURLData(NewPluginURL))
+		UE::GameFeatures::FResult OptionsResult = FGameFeaturePluginState::TryUpdateProtocolOptions(NewOptions);
+		if (OptionsResult.HasError())
 		{
-			return false;
+			return OptionsResult;
+		}
+
+		//if we don't have any in-progress downloads the default behavior is all we need
+		if (PendingBundleDownloads.Num() == 0)
+		{
+			return OptionsResult;
 		}
-		FInstallBundlePluginProtocolMetaData& UpdatedMetaData = StateProperties.ProtocolMetadata.GetSubtype<FInstallBundlePluginProtocolMetaData>();
+
+		const FInstallBundlePluginProtocolOptions& Options = StateProperties.ProtocolOptions.GetSubtype<FInstallBundlePluginProtocolOptions>();
 		
 		//Update our InstallBundleRequestFlags
 		{
-			EInstallBundleRequestFlags UpdatedRequestFlags = OldMetaData.InstallBundleFlags;
+			EInstallBundleRequestFlags UpdatedRequestFlags = Options.InstallBundleFlags;
 						
 			EInstallBundleRequestFlags AddFlags = (UpdatedRequestFlags & (~OldRequestFlags));
 			EInstallBundleRequestFlags RemoveFlags = ((~UpdatedRequestFlags) & OldRequestFlags);
@@ -1277,15 +1270,13 @@ struct FGameFeaturePluginState_Downloading : public FGameFeaturePluginState
 
 		//Handle pausing or resuming the download if the bUserPauseDownload flag has changed
 		{
-			if (UpdatedMetaData.bUserPauseDownload != OldUserPausedFlag)
+			if (Options.bUserPauseDownload != OldUserPausedFlag)
 			{
-				ChangePauseState(UpdatedMetaData.bUserPauseDownload);
+				ChangePauseState(Options.bUserPauseDownload);
 			}
 		}
 
-		//Return true at this point even if we don't change anything explicitly for this state as the 
-		//default FGameFeaturePluginState returned true above and thus changed things
-		return true;
+		return OptionsResult;
 	}
 
 	virtual void EndState() override
@@ -1438,7 +1429,7 @@ struct FGameFeaturePluginState_Unmounting : public FGameFeaturePluginState
 
 		const TArray<FName>& InstallBundles = StateProperties.ProtocolMetadata.GetSubtype<FInstallBundlePluginProtocolMetaData>().InstallBundles;
 
-		EInstallBundleReleaseRequestFlags ReleaseFlags = StateProperties.ProtocolMetadata.GetSubtype<FInstallBundlePluginProtocolMetaData>().ReleaseInstallBundleFlags;
+		EInstallBundleReleaseRequestFlags ReleaseFlags = StateProperties.ProtocolOptions.GetSubtype<FInstallBundlePluginProtocolOptions>().ReleaseInstallBundleFlags;
 		//Make sure we don't remove files here early, that should only be done in Uninstalling
 		ReleaseFlags &= ~(EInstallBundleReleaseRequestFlags::RemoveFilesIfPossible);
 		
@@ -1566,9 +1557,11 @@ struct FGameFeaturePluginState_Mounting : public FGameFeaturePluginState
 		
 		TSharedPtr<IInstallBundleManager> BundleManager = IInstallBundleManager::GetPlatformInstallBundleManager();
 
-		FInstallBundlePluginProtocolMetaData MetaData = StateProperties.ProtocolMetadata.GetSubtype<FInstallBundlePluginProtocolMetaData>();
+		const FInstallBundlePluginProtocolMetaData& MetaData = StateProperties.ProtocolMetadata.GetSubtype<FInstallBundlePluginProtocolMetaData>();
 		const TArray<FName>& InstallBundles = MetaData.InstallBundles;
-		EInstallBundleRequestFlags InstallFlags = MetaData.InstallBundleFlags;
+
+		const FInstallBundlePluginProtocolOptions& Options = StateProperties.ProtocolOptions.GetSubtype<FInstallBundlePluginProtocolOptions>();
+		const EInstallBundleRequestFlags InstallFlags = Options.InstallBundleFlags;
 
 		// Make bundle manager use verbose log level for most logs.
 		// We are already done with downloading, so we don't care about logging too much here unless mounting fails.
@@ -1719,7 +1712,8 @@ struct FGameFeaturePluginState_WaitingForDependencies : public FGameFeaturePlugi
 			UGameFeaturesSubsystem& GameFeaturesSubsystem = UGameFeaturesSubsystem::Get();
 
 			TArray<UGameFeaturePluginStateMachine*> Dependencies;
-			if (!GameFeaturesSubsystem.FindOrCreatePluginDependencyStateMachines(*StateProperties.PluginIdentifier.GetFullPluginURL(), StateProperties.PluginInstalledFilename, Dependencies))
+			if (!GameFeaturesSubsystem.FindOrCreatePluginDependencyStateMachines(
+				*StateProperties.PluginIdentifier.GetFullPluginURL(), StateProperties.PluginInstalledFilename, StateProperties.RecycleProtocolOptions(), Dependencies))
 			{
 				// Failed to query dependencies
 				StateStatus.SetTransitionError(EGameFeaturePluginState::ErrorWaitingForDependencies, GetErrorResult(TEXT("Failed_Dependency_Query")));
@@ -1788,7 +1782,15 @@ struct FGameFeaturePluginState_WaitingForDependencies : public FGameFeaturePlugi
 		// Special case for terminal state since it cannot be exited, we need to make a new machine
 		if (Dependency->GetCurrentState() == EGameFeaturePluginState::Terminal)
 		{
-			UGameFeaturePluginStateMachine* NewMachine = UGameFeaturesSubsystem::Get().FindOrCreateGameFeaturePluginStateMachine(Dependency->GetPluginURL());
+			// Inherit dep protocol options if possible
+			FGameFeatureProtocolOptions DepProtocolOptions;
+			EGameFeaturePluginProtocol DepProtocol = Dependency->GetPluginIdentifier().GetPluginProtocol();
+			if (DepProtocol == EGameFeaturePluginProtocol::InstallBundle && StateProperties.ProtocolOptions.HasSubtype<FInstallBundlePluginProtocolOptions>())
+			{
+				DepProtocolOptions = StateProperties.RecycleProtocolOptions();
+			}
+
+			UGameFeaturePluginStateMachine* NewMachine = UGameFeaturesSubsystem::Get().FindOrCreateGameFeaturePluginStateMachine(Dependency->GetPluginURL(), DepProtocolOptions);
 			checkf(NewMachine != Dependency, TEXT("Game Feature Plugin %s should have already been removed from subsystem!"), *Dependency->GetPluginURL());
 
 			const int32 Index = RemainingDependencies.IndexOfByPredicate([Dependency](const FDepResultPair& Pair)
@@ -2309,7 +2311,7 @@ UGameFeaturePluginStateMachine::UGameFeaturePluginStateMachine(const FObjectInit
 
 }
 
-void UGameFeaturePluginStateMachine::InitStateMachine(FGameFeaturePluginIdentifier InPluginIdentifier)
+void UGameFeaturePluginStateMachine::InitStateMachine(FGameFeaturePluginIdentifier InPluginIdentifier, const FGameFeatureProtocolOptions& InProtocolOptions)
 {
 	check(GetCurrentState() == EGameFeaturePluginState::Uninitialized);
 	CurrentStateInfo.State = EGameFeaturePluginState::UnknownStatus;
@@ -2319,6 +2321,8 @@ void UGameFeaturePluginStateMachine::InitStateMachine(FGameFeaturePluginIdentifi
 		FGameFeaturePluginRequestUpdateStateMachine::CreateUObject(this, &ThisClass::UpdateStateMachine),
 		FGameFeatureStateProgressUpdate::CreateUObject(this, &ThisClass::UpdateCurrentStateProgress));
 
+	StateProperties.ProtocolOptions = InProtocolOptions;
+
 #define GAME_FEATURE_PLUGIN_STATE_MAKE_STATE(inEnum, inText) AllStates[EGameFeaturePluginState::inEnum] = MakeUnique<FGameFeaturePluginState_##inEnum>(StateProperties);
 	GAME_FEATURE_PLUGIN_STATE_LIST(GAME_FEATURE_PLUGIN_STATE_MAKE_STATE)
 #undef GAME_FEATURE_PLUGIN_STATE_MAKE_STATE
@@ -2461,17 +2465,20 @@ bool UGameFeaturePluginStateMachine::TryCancel(FGameFeatureStateTransitionCancel
 	return true;
 }
 
-bool UGameFeaturePluginStateMachine::TryUpdatePluginURLData(const FString& InPluginURL)
+UE::GameFeatures::FResult UGameFeaturePluginStateMachine::TryUpdatePluginProtocolOptions(const FGameFeatureProtocolOptions& InOptions, bool& bOutDidUpdate)
 {
-	//Check if our current PluginURL information matches this new URL
-	//and early out if no update is needed
-	if (StateProperties.PluginIdentifier.ExactMatchesURL(InPluginURL))
+	bOutDidUpdate = false;
+
+	if (StateProperties.ProtocolOptions == InOptions)
 	{
-		return false;
+		return MakeValue();
 	}
 
 	const EGameFeaturePluginState CurrentState = GetCurrentState();
-	return AllStates[CurrentState]->TryUpdateURLData(InPluginURL);
+	UE::GameFeatures::FResult Result = AllStates[CurrentState]->TryUpdateProtocolOptions(InOptions);
+	bOutDidUpdate = Result.HasValue();
+
+	return Result;
 }
 
 void UGameFeaturePluginStateMachine::RemovePendingTransitionCallback(FDelegateHandle InHandle)
@@ -2539,6 +2546,16 @@ const FString& UGameFeaturePluginStateMachine::GetPluginURL() const
 	return StateProperties.PluginIdentifier.GetFullPluginURL();
 }
 
+const FGameFeatureProtocolOptions& UGameFeaturePluginStateMachine::GetProtocolOptions() const
+{
+	return StateProperties.ProtocolOptions;
+}
+
+FGameFeatureProtocolOptions UGameFeaturePluginStateMachine::RecycleProtocolOptions() const
+{
+	return StateProperties.RecycleProtocolOptions();
+}
+
 const FString& UGameFeaturePluginStateMachine::GetPluginName() const
 {
 	return StateProperties.PluginName;
@@ -2780,14 +2797,8 @@ void FInstallBundlePluginProtocolMetaData::ResetToDefaults()
 {
 	InstallBundles.Reset();
 	VersionNum = FDefaultValues::CurrentVersionNum;
-	bUninstallBeforeTerminate = FDefaultValues::Default_bUninstallBeforeTerminate;
-	bUserPauseDownload = FDefaultValues::Default_bUserPauseDownload;
-	bAllowIniLoading = FDefaultValues::Default_bAllowIniLoading;
-	bDoNotDownload = FDefaultValues::Default_bDoNotDownload;
-	InstallBundleFlags = FDefaultValues::Default_InstallBundleFlags;
-	ReleaseInstallBundleFlags = FDefaultValues::Default_ReleaseInstallBundleFlags;
-
-	static_assert(static_cast<uint8>(EGameFeatureInstallBundleProtocolOptions::Count) == 8, "Update this function to handle the newly added EGameFeatureInstallBundleProtocolOptions value!");
+
+	static_assert(static_cast<uint8>(EGameFeatureInstallBundleProtocolOptions::Count) == 2, "Update this function to handle the newly added EGameFeatureInstallBundleProtocolOptions value!");
 }
 
 FString FInstallBundlePluginProtocolMetaData::ToString() const
@@ -2795,51 +2806,28 @@ FString FInstallBundlePluginProtocolMetaData::ToString() const
 	FString ReturnedString;
 
 	//Always encode InstallBundles
-	ReturnedString = LexToString(EGameFeatureInstallBundleProtocolOptions::Bundles) + UE::GameFeatures::PluginURLStructureInfo::OptionAssignOperator;
-	
 	if (InstallBundles.Num() > 0)
 	{
-		const FString BundlesList = FString::JoinBy(InstallBundles, UE::GameFeatures::PluginURLStructureInfo::OptionListSeperator, [](const FName& BundleName){ return FNameBuilder(BundleName); });
+		ReturnedString = LexToString(EGameFeatureInstallBundleProtocolOptions::Bundles) + UE::GameFeatures::PluginURLStructureInfo::OptionAssignOperator;
+
+		FNameBuilder NameBuilder;
+		const FString BundlesList = FString::JoinBy(InstallBundles, UE::GameFeatures::PluginURLStructureInfo::OptionListSeperator, 
+		[&NameBuilder](const FName& BundleName)
+		{
+			BundleName.ToString(NameBuilder);
+			return FStringView(NameBuilder);
+		});
 		ReturnedString.Append(BundlesList);
 	}
 
 	//Always encode version
-	ReturnedString.Append(FString::Printf(TEXT("%s%s%s%d"), UE::GameFeatures::PluginURLStructureInfo::OptionSeperator, *LexToString(EGameFeatureInstallBundleProtocolOptions::Version), UE::GameFeatures::PluginURLStructureInfo::OptionAssignOperator, VersionNum));
-
-	//Encode InstallBundleFlags if not the default
-	if (InstallBundleFlags != FDefaultValues::Default_InstallBundleFlags)
-	{
-		ReturnedString.Append(FString::Printf(TEXT("%s%s%s%d"), UE::GameFeatures::PluginURLStructureInfo::OptionSeperator, *LexToString(EGameFeatureInstallBundleProtocolOptions::Flags), UE::GameFeatures::PluginURLStructureInfo::OptionAssignOperator, static_cast<uint32>(InstallBundleFlags)));
-	}
-
-	if (ReleaseInstallBundleFlags != FDefaultValues::Default_ReleaseInstallBundleFlags)
-	{
-		ReturnedString.Append(FString::Printf(TEXT("%s%s%s%d"), UE::GameFeatures::PluginURLStructureInfo::OptionSeperator, *LexToString(EGameFeatureInstallBundleProtocolOptions::ReleaseFlags), UE::GameFeatures::PluginURLStructureInfo::OptionAssignOperator, static_cast<uint32>(ReleaseInstallBundleFlags)));
-	}
-
-	//Encode bUninstallBeforeTerminate if not the default
-	if (bUninstallBeforeTerminate != FDefaultValues::Default_bUninstallBeforeTerminate)
-	{
-		ReturnedString.Append(FString::Printf(TEXT("%s%s%s%s"), UE::GameFeatures::PluginURLStructureInfo::OptionSeperator, *LexToString(EGameFeatureInstallBundleProtocolOptions::UninstallBeforeTerminate), UE::GameFeatures::PluginURLStructureInfo::OptionAssignOperator, *LexToString(bUninstallBeforeTerminate)));
-	}
-
-	//Encode bUserPauseDownload if not the default
-	if (bUserPauseDownload != FDefaultValues::Default_bUserPauseDownload)
-	{
-		ReturnedString.Append(FString::Printf(TEXT("%s%s%s%s"), UE::GameFeatures::PluginURLStructureInfo::OptionSeperator, *LexToString(EGameFeatureInstallBundleProtocolOptions::UserPauseDownload), UE::GameFeatures::PluginURLStructureInfo::OptionAssignOperator, *LexToString(bUserPauseDownload)));
-	}
+	ReturnedString.Append(FString::Printf(TEXT("%s%s%s%d"), 
+		UE::GameFeatures::PluginURLStructureInfo::OptionSeperator, 
+		*LexToString(EGameFeatureInstallBundleProtocolOptions::Version), 
+		UE::GameFeatures::PluginURLStructureInfo::OptionAssignOperator, 
+		VersionNum));
 
-	if (bAllowIniLoading != FDefaultValues::Default_bAllowIniLoading)
-	{
-		ReturnedString.Append(FString::Printf(TEXT("%s%s%s%s"), UE::GameFeatures::PluginURLStructureInfo::OptionSeperator, *LexToString(EGameFeatureInstallBundleProtocolOptions::AllowIniLoading), UE::GameFeatures::PluginURLStructureInfo::OptionAssignOperator, *LexToString(bAllowIniLoading)));
-	}
-
-	if (bDoNotDownload != FDefaultValues::Default_bDoNotDownload)
-	{
-		ReturnedString.Append(FString::Printf(TEXT("%s%s%s%s"), UE::GameFeatures::PluginURLStructureInfo::OptionSeperator, *LexToString(EGameFeatureInstallBundleProtocolOptions::DoNotDownload), UE::GameFeatures::PluginURLStructureInfo::OptionAssignOperator, *LexToString(bDoNotDownload)));
-	}
-
-	static_assert(static_cast<uint8>(EGameFeatureInstallBundleProtocolOptions::Count) == 8, "Update this function to handle the newly added EGameFeatureInstallBundleProtocolOptions value!");
+	static_assert(static_cast<uint8>(EGameFeatureInstallBundleProtocolOptions::Count) == 2, "Update this function to handle the newly added EGameFeatureInstallBundleProtocolOptions value!");
 
 	return ReturnedString;
 }
@@ -2885,7 +2873,7 @@ bool FInstallBundlePluginProtocolMetaData::FromString(const FString& URLString,
 
 							Metadata.InstallBundles.Reserve(BundleNames.Num());
 							for (FString& BundleNameString : BundleNames)
-		{
+							{
 								Metadata.InstallBundles.Add(*BundleNameString);
 							}
 						}
@@ -2912,98 +2900,6 @@ bool FInstallBundlePluginProtocolMetaData::FromString(const FString& URLString,
 						break;
 					}
 
-					case EGameFeatureInstallBundleProtocolOptions::Flags:
-					{
-						if (OptionStrings.Num() != 2)
-						{
-							bParseSuccess = false;
-							UE_LOG(LogGameFeatures, Warning, TEXT("Error parsing InstallBundle protocol options URL %s . Invalid Flags Option!"), *URLString);
-						}
-						//Expect to parse uint32 flag version of this as a 2nd parameter
-						else
-						{
-							Metadata.InstallBundleFlags = static_cast<EInstallBundleRequestFlags>(FCString::Strtoi(*OptionStrings[1], nullptr, 10));
-						}
-
-						break;
-					}
-
-					case EGameFeatureInstallBundleProtocolOptions::ReleaseFlags:
-					{
-						if (OptionStrings.Num() != 2)
-						{
-							bParseSuccess = false;
-							UE_LOG(LogGameFeatures, Warning, TEXT("Error parsing InstallBundle protocol options URL %s . Invalid ReleaseFlags Option!"), *URLString);
-						}
-						//Expect to parse uint32 flag version of this as a 2nd parameter
-						else
-						{
-							Metadata.ReleaseInstallBundleFlags= static_cast<EInstallBundleReleaseRequestFlags>(FCString::Strtoi(*OptionStrings[1], nullptr, 10));
-						}
-
-						break;
-					}
-
-					case EGameFeatureInstallBundleProtocolOptions::UninstallBeforeTerminate:
-					{
-						if (OptionStrings.Num() != 2)
-						{
-							bParseSuccess = false;
-							UE_LOG(LogGameFeatures, Warning, TEXT("Error parsing InstallBundle protocol options URL %s . Invalid UninstallBeforeTerminate Option!"), *URLString);
-						}
-						else
-						{
-							Metadata.bUninstallBeforeTerminate = OptionStrings[1].ToBool();
-						}
-
-						break;
-					}
-
-					case EGameFeatureInstallBundleProtocolOptions::UserPauseDownload:
-					{
-						if (OptionStrings.Num() != 2)
-						{
-							bParseSuccess = false;
-							UE_LOG(LogGameFeatures, Warning, TEXT("Error parsing InstallBundle protocol options URL %s . Invalid UserPauseDownload Option!"), *URLString);
-						}
-						else
-						{
-							Metadata.bUserPauseDownload = OptionStrings[1].ToBool();
-						}
-
-						break;
-					}
-
-					case EGameFeatureInstallBundleProtocolOptions::AllowIniLoading:
-					{
-						if (OptionStrings.Num() != 2)
-						{
-							bParseSuccess = false;
-							UE_LOG(LogGameFeatures, Warning, TEXT("Error parsing InstallBundle protocol options URL %s . Invalid AllowIniLoading Option!"), *URLString);
-						}
-						else
-						{
-							Metadata.bAllowIniLoading = OptionStrings[1].ToBool();
-						}
-
-						break;
-					}
-
-					case EGameFeatureInstallBundleProtocolOptions::DoNotDownload:
-					{
-						if (OptionStrings.Num() != 2)
-						{
-							bParseSuccess = false;
-							UE_LOG(LogGameFeatures, Warning, TEXT("Error parsing InstallBundle protocol options URL %s . Invalid DoNotDownload Option!"), *URLString);
-						}
-						else
-						{
-							Metadata.bDoNotDownload = OptionStrings[1].ToBool();
-						}
-
-						break;
-					}
-
 					case EGameFeatureInstallBundleProtocolOptions::Count:
 					default:
 					{
@@ -3023,7 +2919,7 @@ bool FInstallBundlePluginProtocolMetaData::FromString(const FString& URLString,
 		UE_LOG(LogGameFeatures, Error, TEXT("Error parsing InstallBundle protocol options URL %s . No Bundle List Found!"), *URLString);
 	}
 
-	static_assert(static_cast<uint8>(EGameFeatureInstallBundleProtocolOptions::Count) == 8, "Update this function to handle the newly added EGameFeatureInstallBundleProtocolOptions value!");
+	static_assert(static_cast<uint8>(EGameFeatureInstallBundleProtocolOptions::Count) == 2, "Update this function to handle the newly added EGameFeatureInstallBundleProtocolOptions value!");
 
 	return bParseSuccess;
 }
@@ -3054,6 +2950,29 @@ bool FGameFeaturePluginStateMachineProperties::ParseURL()
 			ensureMsgf(false, TEXT("Failure to parse URL %s into a valid FInstallBundlePluginProtocolMetaData"), *PluginIdentifier.GetFullPluginURL());
 			return false;
 		}
+
+		// Add default protocol options if they are not set yet
+		if (!ProtocolOptions.HasSubtype<FInstallBundlePluginProtocolOptions>())
+		{
+			if (ProtocolOptions.HasSubtype<FNull>())
+			{
+				ProtocolOptions.SetSubtype<FInstallBundlePluginProtocolOptions>();
+			}
+			else
+			{
+				ensureMsgf(false, TEXT("Protocol options type is incorrect for URL %s"), *PluginIdentifier.GetFullPluginURL());
+				return false;
+			}
+		}
+	}
+	else
+	{
+		// No protocol options for other (file) protocols right now
+		if (!ProtocolOptions.HasSubtype<FNull>())
+		{
+			ensureMsgf(false, TEXT("Protocol options type is incorrect for URL %s"), *PluginIdentifier.GetFullPluginURL());
+			return false;
+		}
 	}
 
 	static_assert(static_cast<uint8>(EGameFeaturePluginProtocol::Count) == 3, "Update FGameFeaturePluginStateMachineProperties::ParseURL to handle any new Metadata parsing required for new EGameFeaturePluginProtocol. If no metadata is required just increment this counter.");
@@ -3061,39 +2980,48 @@ bool FGameFeaturePluginStateMachineProperties::ParseURL()
 	return true;
 }
 
-bool FGameFeaturePluginStateMachineProperties::ValidateURLUpdate(const FGameFeaturePluginStateMachineProperties& OldProperties) const
+UE::GameFeatures::FResult FGameFeaturePluginStateMachineProperties::ValidateProtocolOptionsUpdate(const FGameFeatureProtocolOptions& NewProtocolOptions) const
 {
-	bool bIsValidUpdate = false;
-
-	//Should never change our Identifier or PluginProtocol
-	bIsValidUpdate = ((PluginIdentifier == OldProperties.PluginIdentifier) &&
-						(GetPluginProtocol() == OldProperties.GetPluginProtocol()));
-
-	if (bIsValidUpdate && (GetPluginProtocol() == EGameFeaturePluginProtocol::InstallBundle))
+	if (GetPluginProtocol() == EGameFeaturePluginProtocol::InstallBundle)
 	{
-		if (!ensureAlwaysMsgf((ProtocolMetadata.HasSubtype<FInstallBundlePluginProtocolMetaData>() && 
-								OldProperties.ProtocolMetadata.HasSubtype<FInstallBundlePluginProtocolMetaData>())
-								,TEXT("Error with InstallBundle protocol FGameFeaturePluginStateMachineProperties having an invalid ProtocolMetadata. URL: %s")
+		//Should never change our PluginProtocol
+		if (!ensureAlwaysMsgf(NewProtocolOptions.HasSubtype<FInstallBundlePluginProtocolOptions>()
+								,TEXT("Error with InstallBundle protocol FGameFeaturePluginStateMachineProperties having an invalid ProtocolOptions. URL: %s")
 								,*PluginIdentifier.GetFullPluginURL()))
 		{
-			
-			return false;
+			return MakeError(UE::GameFeatures::StateMachineErrorNamespace + TEXT("ProtocolOptions.Invalid_Protocol"));
 		}
 
-		const FInstallBundlePluginProtocolMetaData& NewMetaData = ProtocolMetadata.GetSubtype<FInstallBundlePluginProtocolMetaData>();
-		const FInstallBundlePluginProtocolMetaData& OldMetaData = OldProperties.ProtocolMetadata.GetSubtype<FInstallBundlePluginProtocolMetaData>();
-
-		if (!ensureMsgf(NewMetaData.InstallBundles == OldMetaData.InstallBundles, TEXT("Unexpected change in InstallBundles when updating URL: %s to %s"), *PluginIdentifier.GetFullPluginURL(), *OldProperties.PluginIdentifier.GetFullPluginURL()))
+		if (ProtocolOptions.HasSubtype<FInstallBundlePluginProtocolOptions>())
 		{
-			bIsValidUpdate = false;
-		}
+			const FInstallBundlePluginProtocolOptions& OldOptions = ProtocolOptions.GetSubtype<FInstallBundlePluginProtocolOptions>();
+			const FInstallBundlePluginProtocolOptions& NewOptions = NewProtocolOptions.GetSubtype<FInstallBundlePluginProtocolOptions>();
 
-		if (!ensureMsgf(NewMetaData.bAllowIniLoading == OldMetaData.bAllowIniLoading, TEXT("Unexpected change to AllowIniLoading when updating URL: %s to %s"), *PluginIdentifier.GetFullPluginURL(), *OldProperties.PluginIdentifier.GetFullPluginURL()))
-		{
-			bIsValidUpdate = false;
+			if (!ensureMsgf(OldOptions.bAllowIniLoading == NewOptions.bAllowIniLoading, TEXT("Unexpected change to AllowIniLoading when updating ProtocolOptions. URL: %s "), *PluginIdentifier.GetFullPluginURL()))
+			{
+				return MakeError(UE::GameFeatures::StateMachineErrorNamespace + TEXT("ProtocolOptions.Invalid_Update"));
+			}
 		}
+
+		return MakeValue();
+	}
+
+	if (NewProtocolOptions.HasSubtype<FNull>())
+	{
+		return MakeValue();
 	}
 
-	return bIsValidUpdate;
+	return MakeError(UE::GameFeatures::StateMachineErrorNamespace + TEXT("ProtocolOptions.Unknown_Protocol"));
+}
+
+FGameFeatureProtocolOptions FGameFeaturePluginStateMachineProperties::RecycleProtocolOptions() const
+{
+	FGameFeatureProtocolOptions Result = ProtocolOptions;
+	if (Result.HasSubtype<FInstallBundlePluginProtocolOptions>())
+	{
+		// Don't allow unexpected uninstalls, otherwise respect any flags previously set by the game
+		Result.GetSubtype<FInstallBundlePluginProtocolOptions>().bUninstallBeforeTerminate = false;
+	}
+	return Result;
 }
 
diff --git a/Engine/Plugins/Experimental/GameFeatures/Source/GameFeatures/Private/GameFeaturePluginStateMachine.h b/Engine/Plugins/Experimental/GameFeatures/Source/GameFeatures/Private/GameFeaturePluginStateMachine.h
--- a/Engine/Plugins/Experimental/GameFeatures/Source/GameFeatures/Private/GameFeaturePluginStateMachine.h
+++ b/Engine/Plugins/Experimental/GameFeatures/Source/GameFeatures/Private/GameFeaturePluginStateMachine.h
@@ -214,6 +214,40 @@ inline bool operator>(const FGameFeaturePluginStateRange& StateRange, EGameFeatu
 	return StateRange.MinState > State;
 }
 
+struct FInstallBundlePluginProtocolMetaData
+{
+	TArray<FName> InstallBundles;
+
+	// TODO: keep this version?
+	/** Set to whatever the FDefaultValues::CurrentVersionNum was when this URL was generated.
+		Allows us to ensure if we try and load URLs generated from a previous version **/
+	uint8 VersionNum;
+
+	/** Functions to convert to/from the URL FString representation of this metadata **/
+	FString ToString() const;
+	static bool FromString(const FString& URLString, FInstallBundlePluginProtocolMetaData& OutMetadata);
+
+	FInstallBundlePluginProtocolMetaData();
+
+private:
+	/** Resets all our Metadata values to the default values */
+	void ResetToDefaults();
+
+	/** Holds default values for the above settings as they are only encoded into a string if they differ from these values */
+	struct FDefaultValues
+	{
+		static const uint32 CurrentVersionNum;
+		//Missing InstallBundles on purpose as the default is just an empty TArray and should always be encoded
+	};
+};
+
+struct FGameFeatureProtocolMetadata : public TUnion<FInstallBundlePluginProtocolMetaData, FNull>
+{
+	FGameFeatureProtocolMetadata() { SetSubtype<FNull>(); }
+	FGameFeatureProtocolMetadata(const FInstallBundlePluginProtocolMetaData& InData) : TUnion(InData) {}
+	FGameFeatureProtocolMetadata(FNull InOptions) { SetSubtype<FNull>(InOptions); }
+};
+
 /** Notification that a state transition is complete */
 DECLARE_DELEGATE_TwoParams(FGameFeatureStateTransitionComplete, UGameFeaturePluginStateMachine* /*Machine*/, const UE::GameFeatures::FResult& /*Result*/);
 
@@ -248,8 +282,11 @@ struct FGameFeaturePluginStateMachineProperties
 	/** Name of the plugin. */
 	FString PluginName;
 
-	/** Meta data parsed from the URL for a specific protocol. */
-	TUnion<FInstallBundlePluginProtocolMetaData> ProtocolMetadata;
+	/** Metadata parsed from the URL for a specific protocol. */
+	FGameFeatureProtocolMetadata ProtocolMetadata;
+
+	/** Additional options for a specific protocol. */
+	FGameFeatureProtocolOptions ProtocolOptions;
 
 	TArray<FName> AddedPrimaryAssetTypes;
 
@@ -288,7 +325,10 @@ struct FGameFeaturePluginStateMachineProperties
 	bool ParseURL();
 
 	/** Checks to see if any invalid data was changed during a URL update. True if data updated was all values expected to be changed. */
-	bool ValidateURLUpdate(const FGameFeaturePluginStateMachineProperties& OldProperties) const;
+	UE::GameFeatures::FResult ValidateProtocolOptionsUpdate(const FGameFeatureProtocolOptions& NewProtocolOptions) const;
+
+	/** Returns protocol options suitable for reuse by another state machine */
+	FGameFeatureProtocolOptions RecycleProtocolOptions() const;
 };
 
 /** Input and output information for a state's UpdateState */
@@ -333,12 +373,9 @@ struct FGameFeaturePluginState
 	/** Attempt to cancel any pending state transition. */
 	virtual void TryCancelState() {}
 
-	/** Called if we have updated the URL for this FGameFeaturePluginState.
-		This can be done whenever URL data is updated that isn't tied to our 
-		FGameFeaturePluginIdentifier information. EX: MetaData options information
-		that is parsed from the URL 
+	/** Called if we have updated the protocol options for this FGameFeaturePluginState.
 		Returns false if no update occured or the update failed. True on successful update. */
-	virtual bool TryUpdateURLData(const FString& NewPluginURL);
+	virtual UE::GameFeatures::FResult TryUpdateProtocolOptions(const FGameFeatureProtocolOptions& NewOptions);
 	
 	/** Called when this state is no longer the active state */
 	virtual void EndState() {}
@@ -421,7 +458,7 @@ public:
 	UGameFeaturePluginStateMachine(const FObjectInitializer& ObjectInitializer);
 
 	/** Initializes the state machine and assigns the URL for the plugin it manages. This sets the machine to the 'UnknownStatus' state. */
-	void InitStateMachine(FGameFeaturePluginIdentifier InPluginIdentifier);
+	void InitStateMachine(FGameFeaturePluginIdentifier InPluginIdentifier, const FGameFeatureProtocolOptions& InProtocolOptions);
 
 	/** Asynchronously transitions the state machine to the destination state range and reports when it is done. 
 	  * DestinationState must be of type EGameFeaturePluginStateType::Destination.
@@ -434,7 +471,7 @@ public:
 
 	/** Update the current PluginURL data for this plugin if possible. Returns false if this update fails or
 		if the supplied InPluginURL matches the existing URL data for this plugin.**/
-	bool TryUpdatePluginURLData(const FString& InPluginURL);
+	UE::GameFeatures::FResult TryUpdatePluginProtocolOptions(const FGameFeatureProtocolOptions& InOptions, bool& bOutDidUpdate);
 
 	/** Remove any pending callback from SetDestination */
 	void RemovePendingTransitionCallback(FDelegateHandle InHandle);
@@ -457,6 +494,12 @@ public:
 	/** Returns the URL */
 	const FString& GetPluginURL() const;
 
+	/** Returns any protocol options */
+	const FGameFeatureProtocolOptions& GetProtocolOptions() const;
+
+	/** Returns protocol options suitable for reuse by another state machine */
+	FGameFeatureProtocolOptions RecycleProtocolOptions() const;
+
 	/** Returns the plugin name if known (plugin must have been registered to know the name). */
 	const FString& GetPluginName() const;
 
diff --git a/Engine/Plugins/Experimental/GameFeatures/Source/GameFeatures/Private/GameFeaturesSubsystem.cpp b/Engine/Plugins/Experimental/GameFeatures/Source/GameFeatures/Private/GameFeaturesSubsystem.cpp
--- a/Engine/Plugins/Experimental/GameFeatures/Source/GameFeatures/Private/GameFeaturesSubsystem.cpp
+++ b/Engine/Plugins/Experimental/GameFeatures/Source/GameFeatures/Private/GameFeaturesSubsystem.cpp
@@ -25,15 +25,6 @@
 
 DEFINE_LOG_CATEGORY(LogGameFeatures);
 
-const uint32 FInstallBundlePluginProtocolMetaData::FDefaultValues::CurrentVersionNum = 1;
-//Missing InstallBundles on purpose as the default is just an empty TArray and should always be encoded
-const bool FInstallBundlePluginProtocolMetaData::FDefaultValues::Default_bUninstallBeforeTerminate = false;
-const bool FInstallBundlePluginProtocolMetaData::FDefaultValues::Default_bUserPauseDownload = false;
-const bool FInstallBundlePluginProtocolMetaData::FDefaultValues::Default_bAllowIniLoading = false;
-const bool FInstallBundlePluginProtocolMetaData::FDefaultValues::Default_bDoNotDownload = false;
-const EInstallBundleRequestFlags FInstallBundlePluginProtocolMetaData::FDefaultValues::Default_InstallBundleFlags = EInstallBundleRequestFlags::Defaults;
-const EInstallBundleReleaseRequestFlags FInstallBundlePluginProtocolMetaData::FDefaultValues::Default_ReleaseInstallBundleFlags = EInstallBundleReleaseRequestFlags::None;
-
 namespace UE::GameFeatures
 {
 	static const FString SubsystemErrorNamespace(TEXT("GameFeaturePlugin.Subsystem."));
@@ -262,6 +253,22 @@ FSimpleDelegate FGameFeatureDeactivatingContext::PauseDeactivationUntilComplete(
 	);
 }
 
+FInstallBundlePluginProtocolOptions::FInstallBundlePluginProtocolOptions()
+	: InstallBundleFlags(EInstallBundleRequestFlags::Defaults)
+	, ReleaseInstallBundleFlags(EInstallBundleReleaseRequestFlags::None)
+{}
+
+bool FInstallBundlePluginProtocolOptions::operator==(const FInstallBundlePluginProtocolOptions& Other) const
+{
+	return
+		InstallBundleFlags == Other.InstallBundleFlags &&
+		ReleaseInstallBundleFlags == Other.ReleaseInstallBundleFlags &&
+		bUninstallBeforeTerminate == Other.bUninstallBeforeTerminate &&
+		bUserPauseDownload == Other.bUserPauseDownload && 
+		bAllowIniLoading == Other.bAllowIniLoading &&
+		bDoNotDownload == Other.bDoNotDownload;
+}
+
 void UGameFeaturesSubsystem::Initialize(FSubsystemCollectionBase& Collection)
 {
 	UE_LOG(LogGameFeatures, Log, TEXT("Initializing game features subsystem"));
@@ -599,6 +606,18 @@ FString UGameFeaturesSubsystem::GetPluginURL_FileProtocol(const FString& PluginD
 	return TEXT("file:") + PluginDescriptorPath;
 }
 
+FString GetPluginURL_InstallBundleProtocol(const FString& PluginName, const FInstallBundlePluginProtocolMetaData& ProtocolMetadata)
+{
+	ensure(ProtocolMetadata.InstallBundles.Num() > 0);
+	FString Path;
+	Path += UE::GameFeatures::GameFeaturePluginProtocolPrefix(EGameFeaturePluginProtocol::InstallBundle);
+	Path += PluginName;
+	Path += UE::GameFeatures::PluginURLStructureInfo::OptionSeperator;
+	Path += ProtocolMetadata.ToString();
+
+	return Path;
+}
+
 FString UGameFeaturesSubsystem::GetPluginURL_InstallBundleProtocol(const FString& PluginName, TArrayView<const FString> BundleNames)
 {
 	FInstallBundlePluginProtocolMetaData ProtocolMetadata;
@@ -606,7 +625,7 @@ FString UGameFeaturesSubsystem::GetPluginURL_InstallBundleProtocol(const FString
 	{
 		ProtocolMetadata.InstallBundles.Add(FName(BundleName));
 	}
-	return GetPluginURL_InstallBundleProtocol(PluginName, ProtocolMetadata);
+	return ::GetPluginURL_InstallBundleProtocol(PluginName, ProtocolMetadata);
 }
 
 FString UGameFeaturesSubsystem::GetPluginURL_InstallBundleProtocol(const FString& PluginName, const FString& BundleName)
@@ -618,7 +637,7 @@ FString UGameFeaturesSubsystem::GetPluginURL_InstallBundleProtocol(const FString
 {
 	FInstallBundlePluginProtocolMetaData ProtocolMetadata;
 	ProtocolMetadata.InstallBundles.Append(BundleNames.GetData(), BundleNames.Num());
-	return GetPluginURL_InstallBundleProtocol(PluginName, ProtocolMetadata);
+	return ::GetPluginURL_InstallBundleProtocol(PluginName, ProtocolMetadata);
 }
 
 FString UGameFeaturesSubsystem::GetPluginURL_InstallBundleProtocol(const FString& PluginName, FName BundleName)
@@ -626,18 +645,6 @@ FString UGameFeaturesSubsystem::GetPluginURL_InstallBundleProtocol(const FString
 	return GetPluginURL_InstallBundleProtocol(PluginName, MakeArrayView(&BundleName, 1));
 }
 
-FString UGameFeaturesSubsystem::GetPluginURL_InstallBundleProtocol(const FString& PluginName, const FInstallBundlePluginProtocolMetaData& ProtocolMetadata)
-{
-	ensure (ProtocolMetadata.InstallBundles.Num() > 0);
-	FString Path;
-	Path += UE::GameFeatures::GameFeaturePluginProtocolPrefix(EGameFeaturePluginProtocol::InstallBundle);
-	Path += PluginName;
-	Path += UE::GameFeatures::PluginURLStructureInfo::OptionSeperator;
-	Path += ProtocolMetadata.ToString();
-
-	return Path;
-}
-
 EGameFeaturePluginProtocol UGameFeaturesSubsystem::GetPluginURLProtocol(FStringView PluginURL)
 {
 	for (EGameFeaturePluginProtocol Protocol : TEnumRange<EGameFeaturePluginProtocol>())
@@ -824,6 +831,11 @@ bool UGameFeaturesSubsystem::IsGameFeaturePluginLoaded(const FString& PluginURL)
 }
 
 void UGameFeaturesSubsystem::LoadGameFeaturePlugin(const FString& PluginURL, const FGameFeaturePluginLoadComplete& CompleteDelegate)
+{
+	LoadGameFeaturePlugin(PluginURL, FGameFeatureProtocolOptions(), CompleteDelegate);
+}
+
+void UGameFeaturesSubsystem::LoadGameFeaturePlugin(const FString& PluginURL, const FGameFeatureProtocolOptions& ProtocolOptions, const FGameFeaturePluginLoadComplete& CompleteDelegate)
 {
 	const bool bIsPluginAllowed = GameSpecificPolicies->IsPluginAllowed(PluginURL);
 	if (!bIsPluginAllowed)
@@ -832,8 +844,8 @@ void UGameFeaturesSubsystem::LoadGameFeaturePlugin(const FString& PluginURL, con
 		return;
 	}
 
-	UGameFeaturePluginStateMachine* StateMachine = FindOrCreateGameFeaturePluginStateMachine(PluginURL);
-	
+	UGameFeaturePluginStateMachine* StateMachine = FindOrCreateGameFeaturePluginStateMachine(PluginURL, ProtocolOptions);
+
 	if (!StateMachine->IsRunning() && StateMachine->GetCurrentState() == EGameFeaturePluginState::Active)
 	{
 		// TODO: Resolve the activated case here, this is needed because in a PIE environment the plugins
@@ -844,7 +856,17 @@ void UGameFeaturesSubsystem::LoadGameFeaturePlugin(const FString& PluginURL, con
 		CallbackObservers(EObserverCallback::Activating, PluginURL, &StateMachine->GetPluginName(), StateMachine->GetGameFeatureDataForActivePlugin());
 	}
 
-	ChangeGameFeatureDestination(StateMachine, FGameFeaturePluginStateRange(EGameFeaturePluginState::Loaded, EGameFeaturePluginState::Active), CompleteDelegate);
+	if (ShouldUpdatePluginProtocolOptions(StateMachine, ProtocolOptions))
+	{
+		const UE::GameFeatures::FResult Result = UpdateGameFeatureProtocolOptions(StateMachine, ProtocolOptions);
+		if (Result.HasError())
+		{
+			CompleteDelegate.ExecuteIfBound(Result);
+			return;
+		}
+	}
+
+	ChangeGameFeatureDestination(StateMachine, ProtocolOptions, FGameFeaturePluginStateRange(EGameFeaturePluginState::Loaded, EGameFeaturePluginState::Active), CompleteDelegate);
 }
 
 void UGameFeaturesSubsystem::LoadAndActivateGameFeaturePlugin(const FString& PluginURL, const FGameFeaturePluginLoadComplete& CompleteDelegate)
@@ -852,7 +874,17 @@ void UGameFeaturesSubsystem::LoadAndActivateGameFeaturePlugin(const FString& Plu
 	ChangeGameFeatureTargetState(PluginURL, EGameFeatureTargetState::Active, CompleteDelegate);
 }
 
+void UGameFeaturesSubsystem::LoadAndActivateGameFeaturePlugin(const FString& PluginURL, const FGameFeatureProtocolOptions& ProtocolOptions, const FGameFeaturePluginLoadComplete& CompleteDelegate)
+{
+	ChangeGameFeatureTargetState(PluginURL, ProtocolOptions, EGameFeatureTargetState::Active, CompleteDelegate);
+}
+
 void UGameFeaturesSubsystem::ChangeGameFeatureTargetState(const FString& PluginURL, EGameFeatureTargetState TargetState, const FGameFeaturePluginChangeStateComplete& CompleteDelegate)
+{
+	ChangeGameFeatureTargetState(PluginURL, FGameFeatureProtocolOptions(), TargetState, CompleteDelegate);
+}
+
+void UGameFeaturesSubsystem::ChangeGameFeatureTargetState(const FString& PluginURL, const FGameFeatureProtocolOptions& ProtocolOptions, EGameFeatureTargetState TargetState, const FGameFeaturePluginChangeStateComplete& CompleteDelegate)
 {
 	EGameFeaturePluginState TargetPluginState = EGameFeaturePluginState::MAX;
 
@@ -884,7 +916,7 @@ void UGameFeaturesSubsystem::ChangeGameFeatureTargetState(const FString& PluginU
 	}
 	else
 	{
-		StateMachine = FindOrCreateGameFeaturePluginStateMachine(PluginURL);
+		StateMachine = FindOrCreateGameFeaturePluginStateMachine(PluginURL, ProtocolOptions);
 	}
 	
 	check(StateMachine);
@@ -912,61 +944,58 @@ void UGameFeaturesSubsystem::ChangeGameFeatureTargetState(const FString& PluginU
 		CallbackObservers(EObserverCallback::Activating, PluginURL, &StateMachine->GetPluginName(), StateMachine->GetGameFeatureDataForActivePlugin());
 	}
 	
-	if (ShouldUpdatePluginURLData(PluginURL))
+	if (ShouldUpdatePluginProtocolOptions(StateMachine, ProtocolOptions))
 	{
-		UpdateGameFeaturePluginURL(PluginURL, FGameFeaturePluginUpdateURLComplete());
+		const UE::GameFeatures::FResult Result = UpdateGameFeatureProtocolOptions(StateMachine, ProtocolOptions);
+		if (Result.HasError())
+		{
+			CompleteDelegate.ExecuteIfBound(Result);
+			return;
+		}
 	}
-
-	ChangeGameFeatureDestination(StateMachine, FGameFeaturePluginStateRange(TargetPluginState), CompleteDelegate);
+	
+	ChangeGameFeatureDestination(StateMachine, ProtocolOptions, FGameFeaturePluginStateRange(TargetPluginState), CompleteDelegate);
 }
 
-void UGameFeaturesSubsystem::UpdateGameFeaturePluginURL(const FString& NewPluginURL)
+UE::GameFeatures::FResult UGameFeaturesSubsystem::UpdateGameFeatureProtocolOptions(const FString& PluginURL, const FGameFeatureProtocolOptions& NewOptions, bool* bOutDidUpdate /*= nullptr*/)
 {
-	UpdateGameFeaturePluginURL(NewPluginURL, FGameFeaturePluginUpdateURLComplete());
-};
+	UGameFeaturePluginStateMachine* StateMachine = FindGameFeaturePluginStateMachine(PluginURL);
+	return UpdateGameFeatureProtocolOptions(StateMachine, NewOptions, bOutDidUpdate);
+}
 
-void UGameFeaturesSubsystem::UpdateGameFeaturePluginURL(const FString& NewPluginURL, const FGameFeaturePluginUpdateURLComplete& CompleteDelegate)
+UE::GameFeatures::FResult UGameFeaturesSubsystem::UpdateGameFeatureProtocolOptions(UGameFeaturePluginStateMachine* StateMachine, const FGameFeatureProtocolOptions& NewOptions, bool* bOutDidUpdate /*= nullptr*/)
 {
-	UGameFeaturePluginStateMachine* StateMachine = nullptr;
-	StateMachine = FindGameFeaturePluginStateMachine(NewPluginURL);
-	if (!StateMachine)
+	if (bOutDidUpdate)
 	{
-		CompleteDelegate.ExecuteIfBound(MakeError(TEXT("GameFeaturePlugin.BadURL")));
-		return;
+		*bOutDidUpdate = false;
 	}
 
-	check(StateMachine);
-
-	const bool bUpdated = StateMachine->TryUpdatePluginURLData(NewPluginURL);
-	if (!bUpdated)
+	if (!StateMachine)
 	{
-		CompleteDelegate.ExecuteIfBound(MakeError(TEXT("GameFeaturePlugin.No_URL_Update_Needed")));
+		return MakeError(UE::GameFeatures::SubsystemErrorNamespace + UE::GameFeatures::CommonErrorCodes::BadURL);
 	}
-	else
-	{
-		FString* PluginURL = GameFeaturePluginNameToPathMap.Find(StateMachine->GetPluginName());
-		if (!ensureAlwaysMsgf(PluginURL, TEXT("Attempt to UpdateGameFeaturePluginURL before GameFeaturePlugin has been added to the GameFeaturePluginNameToPathMap! URL:%s"), *NewPluginURL))
-		{
-			CompleteDelegate.ExecuteIfBound(MakeError(TEXT("GameFeaturePlugin.UpdateTooEarly")));
-			return;
-		}
 
-		*PluginURL = NewPluginURL;
-		CompleteDelegate.ExecuteIfBound(MakeValue());
+	bool bUpdated = false;
+	UE::GameFeatures::FResult Result = StateMachine->TryUpdatePluginProtocolOptions(NewOptions, bUpdated);
+	if (bOutDidUpdate)
+	{
+		*bOutDidUpdate = bUpdated;
 	}
+
+	return Result;
 }
 
-bool UGameFeaturesSubsystem::ShouldUpdatePluginURLData(const FString& NewPluginURL)
+bool UGameFeaturesSubsystem::ShouldUpdatePluginProtocolOptions(const UGameFeaturePluginStateMachine* StateMachine, const FGameFeatureProtocolOptions& NewOptions)
 {
-	UGameFeaturePluginStateMachine* StateMachine = nullptr;
-	StateMachine = FindGameFeaturePluginStateMachine(NewPluginURL);
-	if (!StateMachine || !StateMachine->IsStatusKnown())
+	if (NewOptions.HasSubtype<FNull>())
 	{
 		return false;
 	}
 
-	//Should always be valid at this point
-	check(StateMachine);
+	if (!StateMachine)
+	{
+		return false;
+	}
 	
 	//Make sure our StateMachine isn't in terminal, don't want to update Terminal plugins
 	if (TerminalGameFeaturePluginStateMachines.Contains(StateMachine) || (StateMachine->GetCurrentState() == EGameFeaturePluginState::Terminal))
@@ -974,11 +1003,11 @@ bool UGameFeaturesSubsystem::ShouldUpdatePluginURLData(const FString& NewPluginU
 		return false;
 	}
 
-	if (StateMachine->GetPluginURL().Equals(NewPluginURL, ESearchCase::IgnoreCase))
+	if (StateMachine->GetProtocolOptions() == NewOptions)
 	{
 		return false;
 	}
-
+	
 	return true;
 }
 
@@ -1071,92 +1100,73 @@ void UGameFeaturesSubsystem::ReleaseGameFeaturePlugin(const FString& PluginURL,
 	}
 }
 
-void UGameFeaturesSubsystem::UninstallGameFeaturePlugin(const FString& PluginURL, const FGameFeaturePluginUninstallComplete& CompleteDelegate)
+void UGameFeaturesSubsystem::UninstallGameFeaturePlugin(const FString& PluginURL, const FGameFeaturePluginUninstallComplete& CompleteDelegate /*= FGameFeaturePluginUninstallComplete()*/)
 {
-	//FindOrCreate so that we can make sure we uninstall data for plugins that were installed on a previous application run
-	//but have not yet been requested on this application run and so are not yet in the plugin list but might have data on disk
-	//to uninstall
-	UGameFeaturePluginStateMachine* StateMachine = FindOrCreateGameFeaturePluginStateMachine(PluginURL);
+	UninstallGameFeaturePlugin(PluginURL, FGameFeatureProtocolOptions(), CompleteDelegate);
+}
+
+void UGameFeaturesSubsystem::UninstallGameFeaturePlugin(const FString& PluginURL, const FGameFeatureProtocolOptions& InProtocolOptions, const FGameFeaturePluginUninstallComplete& CompleteDelegate /*= FGameFeaturePluginUninstallComplete()*/)
+{
+	// FindOrCreate so that we can make sure we uninstall data for plugins that were installed on a previous application run
+	// but have not yet been requested on this application run and so are not yet in the plugin list but might have data on disk
+	// to uninstall
+	UGameFeaturePluginStateMachine* StateMachine = FindOrCreateGameFeaturePluginStateMachine(PluginURL, InProtocolOptions);
 	check(StateMachine);
 
-	//We may need to update our PluginURL to force certain metadata changes to facilitate this uninstall
-	//This tracks what URL we actually will pass in to the terminate as we can rely on the fact that Terminate
-	//will update any important Metadata from this new URL before beginning it's terminate
-	FString PluginURLForTerminate = PluginURL;
+	// We may need to update our ProtocolOptions to force certain metadata changes to facilitate this uninstall
+	FGameFeatureProtocolOptions ProtocolOptions = StateMachine->GetProtocolOptions();
 
-	//InstallBundle Protocol GameFeatures may need to change their metadata to force this uninstall
-	if (UGameFeaturesSubsystem::IsPluginURLProtocol(PluginURL, EGameFeaturePluginProtocol::InstallBundle))
+	// InstallBundle Protocol GameFeatures may need to change their metadata to force this uninstall
+	if (StateMachine->GetPluginIdentifier().GetPluginProtocol() == EGameFeaturePluginProtocol::InstallBundle)
 	{
-		// Parse a duplicate version of our current Metadata from the URL
-		FInstallBundlePluginProtocolMetaData ProtocolMetadata;
-		if (!FInstallBundlePluginProtocolMetaData::FromString(PluginURL, ProtocolMetadata))
+		// It's possible that ParseURL hasn't been called yet so setup options here if needed.
+		if (!ProtocolOptions.HasSubtype<FInstallBundlePluginProtocolOptions>())
 		{
-			CompleteDelegate.ExecuteIfBound(UE::GameFeatures::FResult(MakeError(UE::GameFeatures::SubsystemErrorNamespace + UE::GameFeatures::CommonErrorCodes::BadURL)));
+			ensureMsgf(ProtocolOptions.HasSubtype<FNull>(), TEXT("Protocol options type is incorrect for URL %s"), *PluginURL);
+			ProtocolOptions.SetSubtype<FInstallBundlePluginProtocolOptions>();
 		}
 
 		// Need to force on bUninstallBeforeTerminate if it wasn't already set to on in our Metadata
-		if (!ProtocolMetadata.bUninstallBeforeTerminate)
+		FInstallBundlePluginProtocolOptions& InstallBundleOptions = ProtocolOptions.GetSubtype<FInstallBundlePluginProtocolOptions>();
+		if (!InstallBundleOptions.bUninstallBeforeTerminate)
 		{
-			ProtocolMetadata.bUninstallBeforeTerminate = true;
-			
-			FString PluginFilename;
+			InstallBundleOptions.bUninstallBeforeTerminate = true;
+		}
+	}
 
-			//Try and pull PluginFilename from the StateMachine first, but this may not be set yet
-			//as this may be running too early before we have parsed the URL
-			StateMachine->GetPluginFilename(PluginFilename);
-			if (PluginFilename.IsEmpty())
-			{
-				//The PluginIdentifyingString is currently just the PluginFilename so we can fallback to that
-				PluginFilename = StateMachine->GetPluginIdentifier().GetIdentifyingString();
-			}
-			check(!PluginFilename.IsEmpty());
+	// Weird flow here because we need to do a few tasks asynchronously
+	// 1) Update Protocol Options   -->   2) Call to set destination to Uninstall --> 3) After we get to Uninstall go to Terminate
 
-			PluginURLForTerminate = GetPluginURL_InstallBundleProtocol(PluginFilename, ProtocolMetadata);
+	// FIRST:
+	// If we need to update our ProtocolOptions, do that first before starting the Uninstall. This allows us to update
+	// options that might be important on the way to Terminal if they are changed. EX: FInstallBundlePluginProtocolMetaData::bUninstallBeforeTerminate
+	if (ShouldUpdatePluginProtocolOptions(StateMachine, ProtocolOptions))
+	{
+		const UE::GameFeatures::FResult Result = UpdateGameFeatureProtocolOptions(StateMachine, ProtocolOptions);
+		if (Result.HasError())
+		{
+			CompleteDelegate.ExecuteIfBound(Result);
+			return;
 		}
 	}
 
-	//Weird flow here because we need to do a few tasks asynchronously
-	// Update URL   -->    Call to set destination to Uninstall --> After we get to Uninstall go to Terminate
-	// (Called Directly)	   (UninstallTransitionLambda)              (StartTerminateLambda)
-
-	// THIRD:
-	//Lambda that will kick off the actual Terminate after we successfully transition to Uninstalled state
-	const FGameFeaturePluginTerminateComplete StartTerminateLambda = FGameFeaturePluginTerminateComplete::CreateWeakLambda(this,
-		[this, PluginURLForTerminate, CompleteDelegate](const UE::GameFeatures::FResult& Result)
+	// SECOND:
+	// Kick off the Uninstall destination after updating our options if necessary
+	ChangeGameFeatureDestination(StateMachine, ProtocolOptions, FGameFeaturePluginStateRange(EGameFeaturePluginState::Uninstalled),
+		FGameFeaturePluginTerminateComplete::CreateWeakLambda(this, [this, PluginURL, CompleteDelegate](const UE::GameFeatures::FResult& Result)
 		{
+			// THIRD:
+			// Kick off the actual Terminate after we successfully transition to Uninstalled state
 			if (Result.HasValue())
 			{
-				TerminateGameFeaturePlugin(PluginURLForTerminate, CompleteDelegate);
+				TerminateGameFeaturePlugin(PluginURL, CompleteDelegate);
 			}
 			//If we failed just bubble error up
 			else
 			{
 				CompleteDelegate.ExecuteIfBound(Result);
 			}
-		});
-
-	// SECOND:
-	//Lambda that will kick off the Uninstall destination after updating our URL if necessary
-	const FGameFeaturePluginUpdateURLComplete UninstallTransitionLambda = FGameFeaturePluginUpdateURLComplete::CreateWeakLambda(this,
-		[this, PluginURLForTerminate, StartTerminateLambda](const UE::GameFeatures::FResult& Result)
-		{
-			UGameFeaturePluginStateMachine* StateMachine = FindGameFeaturePluginStateMachine(PluginURLForTerminate);
-			check(StateMachine);
-			ChangeGameFeatureDestination(StateMachine, FGameFeaturePluginStateRange(EGameFeaturePluginState::Uninstalled), StartTerminateLambda);
-		});
-
-
-	// FIRST:
-	//If we need to update our URLData, try to do that first before starting the Uninstall. This allows us to update
-	//URL Metadata flags that might be important on the way to Terminal if they are changed. EX: FInstallBundlePluginProtocolMetaData::bUninstallBeforeTerminate
-	if (ShouldUpdatePluginURLData(PluginURLForTerminate))
-	{
-		UpdateGameFeaturePluginURL(PluginURLForTerminate, UninstallTransitionLambda);
-	}
-	else
-	{
-		UninstallTransitionLambda.Execute(MakeValue());
-	}
+		}));
 }
 
 void UGameFeaturesSubsystem::TerminateGameFeaturePlugin(const FString& PluginURL)
@@ -1168,25 +1178,7 @@ void UGameFeaturesSubsystem::TerminateGameFeaturePlugin(const FString& PluginURL
 {
 	if (UGameFeaturePluginStateMachine* StateMachine = FindGameFeaturePluginStateMachine(PluginURL))
 	{
-		//Define a lambda that will kick off the actual Terminate
-		const FGameFeaturePluginUpdateURLComplete StartTerminateLambda = FGameFeaturePluginUpdateURLComplete::CreateWeakLambda(this, 
-			[this, PluginURL, CompleteDelegate](const UE::GameFeatures::FResult& Result)
-			{
-				UGameFeaturePluginStateMachine* StateMachine = FindGameFeaturePluginStateMachine(PluginURL);
-				check(StateMachine);
-				ChangeGameFeatureDestination(StateMachine, FGameFeaturePluginStateRange(EGameFeaturePluginState::Terminal), CompleteDelegate);
-			});
-
-		//If we need to update our URLData, try to do that first before starting the Terminate. This allows us to update
-		//URL Metadata flags that might be important on the way to Terminal if they are changed. EX: FInstallBundlePluginProtocolMetaData::bUninstallBeforeTerminate
-		if (ShouldUpdatePluginURLData(PluginURL))
-		{
-			UpdateGameFeaturePluginURL(PluginURL, StartTerminateLambda);
-		}
-		else
-		{
-			StartTerminateLambda.Execute(MakeValue());
-		}
+		ChangeGameFeatureDestination(StateMachine, FGameFeaturePluginStateRange(EGameFeaturePluginState::Terminal), CompleteDelegate);
 	}
 	else
 	{
@@ -1235,9 +1227,10 @@ void UGameFeaturesSubsystem::LoadBuiltInGameFeaturePlugin(const TSharedRef<IPlug
 			const bool bShouldProcess = AdditionalFilter(Plugin->GetDescriptorFileName(), PluginDetails, BehaviorOptions);
 			if (bShouldProcess)
 			{
-				UGameFeaturePluginStateMachine* StateMachine = FindOrCreateGameFeaturePluginStateMachine(PluginURL);
+				UGameFeaturePluginStateMachine* StateMachine = FindOrCreateGameFeaturePluginStateMachine(PluginURL, FGameFeatureProtocolOptions());
 
-				const EBuiltInAutoState InitialAutoState = (BehaviorOptions.AutoStateOverride != EBuiltInAutoState::Invalid) ? BehaviorOptions.AutoStateOverride : PluginDetails.BuiltInAutoState;
+				const EBuiltInAutoState InitialAutoState = (BehaviorOptions.AutoStateOverride != EBuiltInAutoState::Invalid) ? 
+					BehaviorOptions.AutoStateOverride : PluginDetails.BuiltInAutoState;
 				
 				const EGameFeaturePluginState DestinationState = ConvertInitialFeatureStateToTargetState(InitialAutoState);
 
@@ -1650,29 +1643,43 @@ UGameFeaturePluginStateMachine* UGameFeaturesSubsystem::FindGameFeaturePluginSta
 		GameFeaturePluginStateMachines.FindByHash(GetTypeHash(PluginIdentifier.GetIdentifyingString()), PluginIdentifier.GetIdentifyingString());
 	if (ExistingStateMachine)
 	{
-		UE_LOG(LogGameFeatures, VeryVerbose, TEXT("FOUND GameFeaturePlugin using PluginIdentifier:%.*s for PluginURL:%s"), PluginIdentifier.GetIdentifyingString().Len(), PluginIdentifier.GetIdentifyingString().GetData(), *PluginIdentifier.GetFullPluginURL());
-		return *ExistingStateMachine;
+		EGameFeaturePluginProtocol ExpectedProtocol = (*ExistingStateMachine)->GetPluginIdentifier().GetPluginProtocol();
+		if (ensureMsgf(ExpectedProtocol == PluginIdentifier.GetPluginProtocol(), TEXT("Expected protocol %s for %s"), UE::GameFeatures::GameFeaturePluginProtocolPrefix(ExpectedProtocol), *PluginIdentifier.GetFullPluginURL()))
+		{
+			UE_LOG(LogGameFeatures, VeryVerbose, TEXT("FOUND GameFeaturePlugin using PluginIdentifier:%.*s for PluginURL:%s"), PluginIdentifier.GetIdentifyingString().Len(), PluginIdentifier.GetIdentifyingString().GetData(), *PluginIdentifier.GetFullPluginURL());
+			return *ExistingStateMachine;
+		}
 	}
 	UE_LOG(LogGameFeatures, VeryVerbose, TEXT("NOT FOUND GameFeaturePlugin using PluginIdentifier:%.*s for PluginURL:%s"), PluginIdentifier.GetIdentifyingString().Len(), PluginIdentifier.GetIdentifyingString().GetData(), *PluginIdentifier.GetFullPluginURL());
 
 	return nullptr;
 }
 
-UGameFeaturePluginStateMachine* UGameFeaturesSubsystem::FindOrCreateGameFeaturePluginStateMachine(const FString& PluginURL)
+// Note: ProtocolOptions is not defaulted here. Any API call that could create a state machine should allow the user to pass ProtocolOptions to initialize the machine.
+// It is acceptable that user passes null options. 
+UGameFeaturePluginStateMachine* UGameFeaturesSubsystem::FindOrCreateGameFeaturePluginStateMachine(const FString& PluginURL, const FGameFeatureProtocolOptions& ProtocolOptions)
 {
 	TRACE_CPUPROFILER_EVENT_SCOPE(GFP_FindOrCreateStateMachine);
 	FGameFeaturePluginIdentifier PluginIdentifier(PluginURL);
-	if (UGameFeaturePluginStateMachine* ExistingStateMachine = FindGameFeaturePluginStateMachine(PluginIdentifier))
+	TObjectPtr<UGameFeaturePluginStateMachine> const* ExistingStateMachine =
+		GameFeaturePluginStateMachines.FindByHash(GetTypeHash(PluginIdentifier.GetIdentifyingString()), PluginIdentifier.GetIdentifyingString());
+	if (ExistingStateMachine)
 	{
+		EGameFeaturePluginProtocol ExpectedProtocol = (*ExistingStateMachine)->GetPluginIdentifier().GetPluginProtocol();
+		ensureAlwaysMsgf(ExpectedProtocol == PluginIdentifier.GetPluginProtocol(), TEXT("Expected protocol %s for %s"), UE::GameFeatures::GameFeaturePluginProtocolPrefix(ExpectedProtocol), *PluginIdentifier.GetFullPluginURL());
+
+		// In this case, still return the existing machine, even if the protocol doesn't match. This function should never return null.
+		// There can only be one active instance of any machine.
+
 		UE_LOG(LogGameFeatures, VeryVerbose, TEXT("Found GameFeaturePlugin StateMachine using Identifier:%.*s from PluginURL:%s"), PluginIdentifier.GetIdentifyingString().Len(), PluginIdentifier.GetIdentifyingString().GetData(), *PluginURL);
-		return ExistingStateMachine;
+		return *ExistingStateMachine;
 	}
 
 	UE_LOG(LogGameFeatures, Display, TEXT("Creating GameFeaturePlugin StateMachine using Identifier:%.*s from PluginURL:%s"), PluginIdentifier.GetIdentifyingString().Len(), PluginIdentifier.GetIdentifyingString().GetData(), *PluginURL);
 
 	UGameFeaturePluginStateMachine* NewStateMachine = NewObject<UGameFeaturePluginStateMachine>(this);
 	GameFeaturePluginStateMachines.Add(FString(PluginIdentifier.GetIdentifyingString()), NewStateMachine);
-	NewStateMachine->InitStateMachine(MoveTemp(PluginIdentifier));
+	NewStateMachine->InitStateMachine(MoveTemp(PluginIdentifier), ProtocolOptions);
 
 	return NewStateMachine;
 }
@@ -1708,6 +1715,11 @@ void UGameFeaturesSubsystem::LoadBuiltInGameFeaturePluginComplete(const UE::Game
 }
 
 void UGameFeaturesSubsystem::ChangeGameFeatureDestination(UGameFeaturePluginStateMachine* Machine, const FGameFeaturePluginStateRange& StateRange, FGameFeaturePluginChangeStateComplete CompleteDelegate)
+{
+	ChangeGameFeatureDestination(Machine, FGameFeatureProtocolOptions(), StateRange, CompleteDelegate);
+}
+
+void UGameFeaturesSubsystem::ChangeGameFeatureDestination(UGameFeaturePluginStateMachine* Machine, const FGameFeatureProtocolOptions& InProtocolOptions, const FGameFeaturePluginStateRange& StateRange, FGameFeaturePluginChangeStateComplete CompleteDelegate)
 {
 	const bool bSetDestination = Machine->SetDestination(StateRange,
 		FGameFeatureStateTransitionComplete::CreateUObject(this, &ThisClass::ChangeGameFeatureTargetStateComplete, CompleteDelegate));
@@ -1725,12 +1737,12 @@ void UGameFeaturesSubsystem::ChangeGameFeatureDestination(UGameFeaturePluginStat
 			*UE::GameFeatures::ToString(CurrDesitination.MinState), *UE::GameFeatures::ToString(CurrDesitination.MaxState));
 
 		// Try canceling any current transition, then retry
-		auto OnCanceled = [this, StateRange, CompleteDelegate](UGameFeaturePluginStateMachine* Machine) mutable
+		auto OnCanceled = [this, InProtocolOptions, StateRange, CompleteDelegate](UGameFeaturePluginStateMachine* Machine) mutable
 		{
 			// Special case for terminal state since it cannot be exited, we need to make a new machine
 			if (Machine->GetCurrentState() == EGameFeaturePluginState::Terminal)
 			{
-				UGameFeaturePluginStateMachine* NewMachine = FindOrCreateGameFeaturePluginStateMachine(Machine->GetPluginURL());
+				UGameFeaturePluginStateMachine* NewMachine = FindOrCreateGameFeaturePluginStateMachine(Machine->GetPluginURL(), InProtocolOptions);
 				checkf(NewMachine != Machine, TEXT("Game Feature Plugin %s should have already been removed from subsystem!"), *Machine->GetPluginURL());
 				Machine = NewMachine;
 			}
@@ -1784,14 +1796,22 @@ void UGameFeaturesSubsystem::FinishTermination(UGameFeaturePluginStateMachine* M
 	TerminalGameFeaturePluginStateMachines.RemoveSwap(Machine);
 }
 
-bool UGameFeaturesSubsystem::FindOrCreatePluginDependencyStateMachines(const FString& PluginURL, const FString& PluginFilename, TArray<UGameFeaturePluginStateMachine*>& OutDependencyMachines)
+bool UGameFeaturesSubsystem::FindOrCreatePluginDependencyStateMachines(const FString& PluginURL, const FString& PluginFilename, const FGameFeatureProtocolOptions& InDepProtocolOptions, TArray<UGameFeaturePluginStateMachine*>& OutDependencyMachines)
 {
 	FGameFeaturePluginDetails Details;
 	if (GetGameFeaturePluginDetails(PluginURL, PluginFilename, Details))
 	{
 		for (const FString& DependencyURL : Details.PluginDependencies)
 		{
-			UGameFeaturePluginStateMachine* Dependency = FindOrCreateGameFeaturePluginStateMachine(DependencyURL);
+			// Inherit dep protocol options if possible
+			FGameFeatureProtocolOptions DepProtocolOptions;
+			EGameFeaturePluginProtocol DepProtocol = UGameFeaturesSubsystem::GetPluginURLProtocol(DependencyURL);
+			if (DepProtocol == EGameFeaturePluginProtocol::InstallBundle && InDepProtocolOptions.HasSubtype<FInstallBundlePluginProtocolOptions>())
+			{
+				DepProtocolOptions = InDepProtocolOptions;
+			}
+
+			UGameFeaturePluginStateMachine* Dependency = FindOrCreateGameFeaturePluginStateMachine(DependencyURL, DepProtocolOptions);
 			check(Dependency);
 			OutDependencyMachines.Add(Dependency);
 		}
diff --git a/Engine/Plugins/Experimental/GameFeatures/Source/GameFeatures/Public/GameFeaturesSubsystem.h b/Engine/Plugins/Experimental/GameFeatures/Source/GameFeatures/Public/GameFeaturesSubsystem.h
--- a/Engine/Plugins/Experimental/GameFeatures/Source/GameFeatures/Public/GameFeaturesSubsystem.h
+++ b/Engine/Plugins/Experimental/GameFeatures/Source/GameFeatures/Public/GameFeaturesSubsystem.h
@@ -2,6 +2,7 @@
 
 #pragma once
 
+#include "Containers/Union.h"
 #include "Engine/Engine.h"
 #include "GameFeatureTypesFwd.h"
 
@@ -164,7 +165,7 @@ using FGameFeaturePluginUnloadComplete = FGameFeaturePluginChangeStateComplete;
 using FGameFeaturePluginReleaseComplete = FGameFeaturePluginChangeStateComplete;
 using FGameFeaturePluginUninstallComplete = FGameFeaturePluginChangeStateComplete;
 using FGameFeaturePluginTerminateComplete = FGameFeaturePluginChangeStateComplete;
-using FGameFeaturePluginUpdateURLComplete = FGameFeaturePluginChangeStateComplete;
+using FGameFeaturePluginUpdateProtocolComplete = FGameFeaturePluginChangeStateComplete;
 
 DECLARE_DELEGATE_OneParam(FBuiltInGameFeaturePluginsLoaded, PREPROCESSOR_COMMA_SEPARATED(const TMap<FString, UE::GameFeatures::FResult>& /*Results*/));
 
@@ -270,18 +271,11 @@ private:
 };
 
 USTRUCT()
-struct GAMEFEATURES_API FInstallBundlePluginProtocolMetaData
+struct GAMEFEATURES_API FInstallBundlePluginProtocolOptions
 {
 	GENERATED_BODY()
 
-	TArray<FName> InstallBundles;
-
-	/** Set to whatever the FDefaultValues::CurrentVersionNum was when this URL was generated.
-		Allows us to ensure if we try and load URLs generated from a previous version **/
-	uint8 VersionNum;
-
-	/** If we want to attempt to uninstall InstallBundle information installed by this plugin before terminating */
-	bool bUninstallBeforeTerminate;
+	FInstallBundlePluginProtocolOptions();
 
 	/** EInstallBundleRequestFlags utilized during the download/install by InstallBundleManager */
 	EInstallBundleRequestFlags InstallBundleFlags;
@@ -289,37 +283,26 @@ struct GAMEFEATURES_API FInstallBundlePluginProtocolMetaData
 	/** EInstallBundleReleaseRequestFlags utilized during our release and uninstall states */
 	EInstallBundleReleaseRequestFlags ReleaseInstallBundleFlags;
 
+	/** If we want to attempt to uninstall InstallBundle data installed by this plugin before terminating */
+	bool bUninstallBeforeTerminate = false;
+
 	/** If we want to set the Downloading state to pause because of user interaction */
-	bool bUserPauseDownload;
+	bool bUserPauseDownload = false;
 
 	/** Allow the GFP to load INI files, should only be allowed for trusted content */
-	bool bAllowIniLoading;
+	bool bAllowIniLoading = false;
 
-	/** Functions to convert to/from the URL FString representation of this metadata **/
-	FString ToString() const;
-	static bool FromString(const FString& URLString, FInstallBundlePluginProtocolMetaData& OutMetadata);
+	/** Disallows downloading, useful for conditionally loading content only if it's already been installed **/
+	bool bDoNotDownload = false;
 
-	/** Resets all our Metadata values to the default values */
-	void ResetToDefaults();
-
-	/** disallows downloading using the bundle manager **/
-	bool bDoNotDownload;
-
-	FInstallBundlePluginProtocolMetaData();
+	bool operator==(const FInstallBundlePluginProtocolOptions& Other) const;
+};
 
-private:
-	/** Holds default values for the above settings as they are only encoded into a string if they differ from these values */
-	struct FDefaultValues
-	{
-		static const uint32 CurrentVersionNum;
-		//Missing InstallBundles on purpose as the default is just an empty TArray and should always be encoded
-		static const bool Default_bUninstallBeforeTerminate;
-		static const bool Default_bUserPauseDownload;
-		static const bool Default_bAllowIniLoading;
-		static const bool Default_bDoNotDownload;
-		static const EInstallBundleRequestFlags Default_InstallBundleFlags;
-		static const EInstallBundleReleaseRequestFlags Default_ReleaseInstallBundleFlags;
-	};
+struct FGameFeatureProtocolOptions : public TUnion<FInstallBundlePluginProtocolOptions, FNull>
+{
+	FGameFeatureProtocolOptions() { SetSubtype<FNull>(); }
+	FGameFeatureProtocolOptions(const FInstallBundlePluginProtocolOptions& InOptions) : TUnion(InOptions) {}
+	FGameFeatureProtocolOptions(FNull InOptions) { SetSubtype<FNull>(InOptions); }
 };
 
 /** The manager subsystem for game features */
@@ -391,7 +374,6 @@ public:
 	static FString GetPluginURL_InstallBundleProtocol(const FString& PluginName, const FString& BundleName);
 	static FString GetPluginURL_InstallBundleProtocol(const FString& PluginName, TArrayView<const FName> BundleNames);
 	static FString GetPluginURL_InstallBundleProtocol(const FString& PluginName, FName BundleName);
-	static FString GetPluginURL_InstallBundleProtocol(const FString& PluginName, const FInstallBundlePluginProtocolMetaData& ProtocolMetadata);
 
 	/** Returns the plugin protocol for the specified URL */
 	static EGameFeaturePluginProtocol GetPluginURLProtocol(FStringView PluginURL);
@@ -420,16 +402,18 @@ public:
 
 	/** Loads a single game feature plugin. */
 	void LoadGameFeaturePlugin(const FString& PluginURL, const FGameFeaturePluginLoadComplete& CompleteDelegate);
+	void LoadGameFeaturePlugin(const FString& PluginURL, const FGameFeatureProtocolOptions& ProtocolOptions, const FGameFeaturePluginLoadComplete& CompleteDelegate);
 
 	/** Loads a single game feature plugin and activates it. */
 	void LoadAndActivateGameFeaturePlugin(const FString& PluginURL, const FGameFeaturePluginLoadComplete& CompleteDelegate);
+	void LoadAndActivateGameFeaturePlugin(const FString& PluginURL, const FGameFeatureProtocolOptions& ProtocolOptions, const FGameFeaturePluginLoadComplete& CompleteDelegate);
 
 	/** Changes the target state of a game feature plugin */
 	void ChangeGameFeatureTargetState(const FString& PluginURL, EGameFeatureTargetState TargetState, const FGameFeaturePluginChangeStateComplete& CompleteDelegate);
+	void ChangeGameFeatureTargetState(const FString& PluginURL, const FGameFeatureProtocolOptions& ProtocolOptions, EGameFeatureTargetState TargetState, const FGameFeaturePluginChangeStateComplete& CompleteDelegate);
 
-	/** Changes the URL data of a game feature plugin. Useful to change any options data that is parsed from the URL such as settings flags */
-	void UpdateGameFeaturePluginURL(const FString& NewPluginURL);
-	void UpdateGameFeaturePluginURL(const FString& NewPluginURL, const FGameFeaturePluginUpdateURLComplete& CompleteDelegate);
+	/** Changes the protocol options of a game feature plugin. Useful to change any options data such as settings flags */
+	UE::GameFeatures::FResult UpdateGameFeatureProtocolOptions(const FString& PluginURL, const FGameFeatureProtocolOptions& NewOptions, bool* bOutDidUpdate = nullptr);
 
 	/** Gets the Install_Percent for single game feature plugin if it is active. */
 	bool GetGameFeaturePluginInstallPercent(const FString& PluginURL, float& Install_Percent) const;
@@ -453,8 +437,8 @@ public:
 		If the given PluginURL is not found this will create a GameFeaturePlugin first and attempt to run it through the uninstall flow.
 		This allows for the uninstalling of data that was installed on previous runs of the application where we haven't yet requested the
 		GameFeaturePlugin that we would like to uninstall data for on this run. */
-	void UninstallGameFeaturePlugin(const FString& PluginURL);
-	void UninstallGameFeaturePlugin(const FString& PluginURL, const FGameFeaturePluginUninstallComplete& CompleteDelegate);
+	void UninstallGameFeaturePlugin(const FString& PluginURL, const FGameFeaturePluginUninstallComplete& CompleteDelegate = FGameFeaturePluginUninstallComplete());
+	void UninstallGameFeaturePlugin(const FString& PluginURL, const FGameFeatureProtocolOptions& ProtocolOptions, const FGameFeaturePluginUninstallComplete& CompleteDelegate = FGameFeaturePluginUninstallComplete());
 
 	/** Terminate the GameFeaturePlugin and remove all associated plugin tracking data. */
 	void TerminateGameFeaturePlugin(const FString& PluginURL);
@@ -558,7 +542,8 @@ private:
 	static void RemoveGameFeatureFromAssetManager(const UGameFeatureData* GameFeatureToRemove, const FString& PluginName, const TArray<FName>& AddedPrimaryAssetTypes);
 
 private:
-	bool ShouldUpdatePluginURLData(const FString& NewPluginURL);
+	bool ShouldUpdatePluginProtocolOptions(const UGameFeaturePluginStateMachine* StateMachine, const FGameFeatureProtocolOptions& NewOptions);
+	UE::GameFeatures::FResult UpdateGameFeatureProtocolOptions(UGameFeaturePluginStateMachine* StateMachine, const FGameFeatureProtocolOptions& NewOptions, bool* bOutDidUpdate = nullptr);
 
 	const UGameFeatureData* GetDataForStateMachine(UGameFeaturePluginStateMachine* GFSM) const;
 	const UGameFeatureData* GetRegisteredDataForStateMachine(UGameFeaturePluginStateMachine* GFSM) const;
@@ -580,13 +565,17 @@ private:
 	UGameFeaturePluginStateMachine* FindGameFeaturePluginStateMachine(const FGameFeaturePluginIdentifier& PluginIdentifier) const;
 
 	/** Gets the state machine associated with the specified URL, creates it if it doesnt exist */
-	UGameFeaturePluginStateMachine* FindOrCreateGameFeaturePluginStateMachine(const FString& PluginURL);
+	UGameFeaturePluginStateMachine* FindOrCreateGameFeaturePluginStateMachine(const FString& PluginURL, const FGameFeatureProtocolOptions& ProtocolOptions);
 
 	/** Notification that a game feature has finished loading, and whether it was successful */
 	void LoadBuiltInGameFeaturePluginComplete(const UE::GameFeatures::FResult& Result, UGameFeaturePluginStateMachine* Machine, FGameFeaturePluginStateRange RequestedDestination);
 
-	/** Sets a new destination state. Will attempt to cancel the current transition if the new destination is incompatible with the current destination */
+	/** 
+	 * Sets a new destination state. Will attempt to cancel the current transition if the new destination is incompatible with the current destination 
+	 * Note: In the case that the existing machine is terminal, a new one will need to be created. In that case ProtocolOptions will be used for the new machine.
+	 */
 	void ChangeGameFeatureDestination(UGameFeaturePluginStateMachine* Machine, const FGameFeaturePluginStateRange& StateRange, FGameFeaturePluginChangeStateComplete CompleteDelegate);
+	void ChangeGameFeatureDestination(UGameFeaturePluginStateMachine* Machine, const FGameFeatureProtocolOptions& ProtocolOptions, const FGameFeaturePluginStateRange& StateRange, FGameFeaturePluginChangeStateComplete CompleteDelegate);
 
 	/** Generic notification that calls the Complete delegate without broadcasting anything else.*/
 	void ChangeGameFeatureTargetStateComplete(UGameFeaturePluginStateMachine* Machine, const UE::GameFeatures::FResult& Result, FGameFeaturePluginChangeStateComplete CompleteDelegate);
@@ -596,7 +585,7 @@ private:
 	friend class UGameFeaturePluginStateMachine;
 
 	/** Handler for when a state machine requests its dependencies. Returns false if the dependencies could not be read */
-	bool FindOrCreatePluginDependencyStateMachines(const FString& PluginURL, const FString& PluginFilename, TArray<UGameFeaturePluginStateMachine*>& OutDependencyMachines);
+	bool FindOrCreatePluginDependencyStateMachines(const FString& PluginURL, const FString& PluginFilename, const FGameFeatureProtocolOptions& DepProtocolOptions, TArray<UGameFeaturePluginStateMachine*>& OutDependencyMachines);
 	friend struct FGameFeaturePluginState_WaitingForDependencies;
 
 	/** Handle 'ListGameFeaturePlugins' console command */
@@ -635,7 +624,7 @@ private:
 	{
 		FGameFeaturePluginDetails Details;
 		FDateTime TimeStamp;
-		FCachedGameFeaturePluginDetails() {}
+		FCachedGameFeaturePluginDetails() = default;
 		FCachedGameFeaturePluginDetails(const FGameFeaturePluginDetails& InDetails, const FDateTime& InTimeStamp) : Details(InDetails), TimeStamp(InTimeStamp) {}
 	};
 	mutable TMap<FString, FCachedGameFeaturePluginDetails> CachedPluginDetailsByFilename;
