GFP URL can encode game specific data
- Add new overrides to GFP URL utility functions
- Pass through unknown URL options without an error
- Use short URLs for logging where possible
- Remove version field from install bundle protocol

diff --git a/Engine/Plugins/Experimental/GameFeatures/Source/GameFeatures/Private/GameFeaturePluginStateMachine.cpp b/Engine/Plugins/Experimental/GameFeatures/Source/GameFeatures/Private/GameFeaturePluginStateMachine.cpp
--- a/Engine/Plugins/Experimental/GameFeatures/Source/GameFeatures/Private/GameFeaturePluginStateMachine.cpp
+++ b/Engine/Plugins/Experimental/GameFeatures/Source/GameFeatures/Private/GameFeaturePluginStateMachine.cpp
@@ -15,6 +15,7 @@
 #include "Misc/AsciiSet.h"
 #include "Misc/ConfigCacheIni.h"
 #include "Misc/CoreDelegates.h"
+#include "Misc/EnumRange.h"
 #include "Misc/FileHelper.h"
 #include "Misc/PackageName.h"
 #include "Misc/WildcardString.h"
@@ -23,6 +24,7 @@
 #include "RenderDeferredCleanup.h"
 #include "Serialization/MemoryReader.h"
 #include "String/ParseTokens.h"
+#include "String/LexFromString.h"
 #include "GameFeaturesProjectPolicies.h"
 #include "UObject/ObjectRename.h"
 #include "UObject/ReferenceChainSearch.h"
@@ -38,8 +40,6 @@
 #include "Misc/App.h"
 #endif //if WITH_EDITOR
 
-const uint32 FInstallBundlePluginProtocolMetaData::FDefaultValues::CurrentVersionNum = 2;
-
 namespace UE::GameFeatures
 {
 	static const FString StateMachineErrorNamespace(TEXT("GameFeaturePlugin.StateMachine."));
@@ -261,8 +261,7 @@ namespace UE::GameFeatures
 //Enum describing the options for what data can be inside the InstallBundle protocol URL's metadata
 //Used to aid in parsing / creating the URL for InstallBundle protocol GameFeaturePlugins
 #define INSTALL_BUNDLE_PROTOCOL_URL_OPTIONS_LIST(XOPTION)					\
-	XOPTION(Bundles,					TEXT("Bundles"))					\
-	XOPTION(Version,					TEXT("V"))							
+	XOPTION(Bundles,					TEXT("Bundles"))					
 
 #define INSTALL_BUNDLE_PROTOCOL_URL_OPTIONS_ENUM(inEnum, inString) inEnum,
 enum class EGameFeatureInstallBundleProtocolOptions : uint8
@@ -270,6 +269,7 @@ enum class EGameFeatureInstallBundleProtocolOptions : uint8
 	INSTALL_BUNDLE_PROTOCOL_URL_OPTIONS_LIST(INSTALL_BUNDLE_PROTOCOL_URL_OPTIONS_ENUM)
 	Count,
 };
+ENUM_RANGE_BY_COUNT(EGameFeatureInstallBundleProtocolOptions, EGameFeatureInstallBundleProtocolOptions::Count);
 #undef INSTALL_BUNDLE_PROTOCOL_URL_OPTIONS_ENUM
 
 #define INSTALL_BUNDLE_PROTOCOL_URL_OPTIONS_LEX_TO_STRING(inEnum, inString) \
@@ -315,28 +315,25 @@ void FGameFeaturePluginStateStatus::SetTransitionError(EGameFeaturePluginState T
 	}
 }
 
-void LexFromString(EGameFeatureInstallBundleProtocolOptions& ValueOut, const TCHAR* StringIn)
+void LexFromString(EGameFeatureInstallBundleProtocolOptions& ValueOut, const FStringView StringIn)
 {
 	//Default value if parsing fails
 	ValueOut = EGameFeatureInstallBundleProtocolOptions::Count;
 
-	if (!StringIn || (StringIn[0] == '\0'))
+	if (StringIn.IsEmpty() || StringIn[0] == TEXT('\0'))
 	{
 		UE_LOG(LogGameFeatures, Error, TEXT("Invalid empty string used for EGameFeatureInstallBundleProtocolOptions LexFromString!"));
 		return;
 	}
 
-	for (uint8 EnumIndex = 0; EnumIndex < static_cast<uint8>(EGameFeatureInstallBundleProtocolOptions::Count); ++EnumIndex)
+	for (EGameFeatureInstallBundleProtocolOptions OptionToCheck : TEnumRange<EGameFeatureInstallBundleProtocolOptions>())
 	{
-		EGameFeatureInstallBundleProtocolOptions GameFeatureToCheck = static_cast<EGameFeatureInstallBundleProtocolOptions>(EnumIndex);
-		if (LexToString(GameFeatureToCheck).Equals(StringIn,ESearchCase::IgnoreCase))
+		if (LexToString(OptionToCheck) == StringIn)
 		{
-			ValueOut = GameFeatureToCheck;
+			ValueOut = OptionToCheck;
 			return;
 		}
 	}
-
-	UE_LOG(LogGameFeatures, Error, TEXT("No valid EGameFeatureInstallBundleProtocolOptions to LexFromString from %s"), StringIn);
 }
 
 UE::GameFeatures::FResult FGameFeaturePluginState::GetErrorResult(const FString& ErrorCode, const FText OptionalErrorText/*= FText()*/) const
@@ -789,7 +786,8 @@ struct FBaseDataReleaseGameFeaturePluginState : public FGameFeaturePluginState
 
 		if (MaybeRequestInfo.HasError())
 		{
-			ensureMsgf(false, TEXT("Unable to enqueue uninstall for the PluginURL(%s) because %s"), *StateProperties.PluginIdentifier.GetFullPluginURL(), LexToString(MaybeRequestInfo.GetError()));
+			const FStringView ShortUrl = StateProperties.PluginIdentifier.GetIdentifyingString();
+			ensureMsgf(false, TEXT("Unable to enqueue uninstall for the PluginURL(%.*s) because %s"), ShortUrl.Len(), ShortUrl.GetData(), LexToString(MaybeRequestInfo.GetError()));
 			Result = GetErrorResult(TEXT("BundleManager.Begin."), MaybeRequestInfo.GetError());
 
 			return;
@@ -799,7 +797,8 @@ struct FBaseDataReleaseGameFeaturePluginState : public FGameFeaturePluginState
 
 		if (EnumHasAnyFlags(RequestInfo.InfoFlags, EInstallBundleRequestInfoFlags::SkippedUnknownBundles))
 		{
-			ensureMsgf(false, TEXT("Unable to enqueue uninstall for the PluginURL(%s) because failed to resolve install bundles!"), *StateProperties.PluginIdentifier.GetFullPluginURL());
+			const FStringView ShortUrl = StateProperties.PluginIdentifier.GetIdentifyingString();
+			ensureMsgf(false, TEXT("Unable to enqueue uninstall for the PluginURL(%.*s) because failed to resolve install bundles!"), ShortUrl.Len(), ShortUrl.GetData());
 			Result = GetErrorResult(TEXT("BundleManager.Begin."), TEXT("Resolve_Failed"), UE::GameFeatures::CommonErrorCodes::ReleaseResult_Generic);
 
 			return;
@@ -1054,7 +1053,8 @@ struct FGameFeaturePluginState_Downloading : public FGameFeaturePluginState
 
 		if (MaybeRequestInfo.HasError())
 		{
-			ensureMsgf(false, TEXT("Unable to enqueue download for the PluginURL(%s) because %s"), *StateProperties.PluginIdentifier.GetFullPluginURL(), LexToString(MaybeRequestInfo.GetError()));
+			const FStringView ShortUrl = StateProperties.PluginIdentifier.GetIdentifyingString();
+			ensureMsgf(false, TEXT("Unable to enqueue download for the PluginURL(%.*s) because %s"), ShortUrl.Len(), ShortUrl.GetData(), LexToString(MaybeRequestInfo.GetError()));
 			Result = GetErrorResult(TEXT("BundleManager.GotState."), LexToString(MaybeRequestInfo.GetError()));
 			UpdateStateMachineImmediate();
 			return;
@@ -1064,7 +1064,8 @@ struct FGameFeaturePluginState_Downloading : public FGameFeaturePluginState
 
 		if (EnumHasAnyFlags(RequestInfo.InfoFlags, EInstallBundleRequestInfoFlags::SkippedUnknownBundles))
 		{
-			ensureMsgf(false, TEXT("Unable to enqueue download for the PluginURL(%s) because failed to resolve install bundles!"), *StateProperties.PluginIdentifier.GetFullPluginURL());
+			const FStringView ShortUrl = StateProperties.PluginIdentifier.GetIdentifyingString();
+			ensureMsgf(false, TEXT("Unable to enqueue download for the PluginURL(%.*s) because failed to resolve install bundles!"), ShortUrl.Len(), ShortUrl.GetData());
 			Result = GetErrorResult(TEXT("BundleManager.GotState."), TEXT("Resolve_Failed"), UE::GameFeatures::CommonErrorCodes::Generic_ConnectionError);
 			
 			UpdateStateMachineImmediate();
@@ -1146,7 +1147,8 @@ struct FGameFeaturePluginState_Downloading : public FGameFeaturePluginState
 			float Progress = ProgressTracker->GetCurrentCombinedProgress().ProgressPercent;
 			UpdateProgress(Progress);
 
-			UE_LOG(LogGameFeatures, VeryVerbose, TEXT("Download Progress: %f for PluginURL(%s)"), Progress, *StateProperties.PluginIdentifier.GetFullPluginURL());
+			const FStringView ShortUrl = StateProperties.PluginIdentifier.GetIdentifyingString();
+			UE_LOG(LogGameFeatures, VeryVerbose, TEXT("Download Progress: %f for PluginURL(%.*s)"), Progress, ShortUrl.Len(), ShortUrl.GetData());
 		}
 
 		return true;
@@ -1405,7 +1407,8 @@ struct FGameFeaturePluginState_Unmounting : public FGameFeaturePluginState
 			FText FailureReason;
 			if (!IPluginManager::Get().UnmountExplicitlyLoadedPlugin(StateProperties.PluginName, &FailureReason))
 			{
-				ensureMsgf(false, TEXT("Failed to explicitly unmount the PluginURL(%s) because %s"), *StateProperties.PluginIdentifier.GetFullPluginURL(), *FailureReason.ToString());
+				const FStringView ShortUrl = StateProperties.PluginIdentifier.GetIdentifyingString();
+				ensureMsgf(false, TEXT("Failed to explicitly unmount the PluginURL(%.*s) because %s"), ShortUrl.Len(), ShortUrl.GetData(), *FailureReason.ToString());
 				Result = GetErrorResult(TEXT("Plugin_Cannot_Explicitly_Unmount"));
 				return;
 			}
@@ -1437,7 +1440,8 @@ struct FGameFeaturePluginState_Unmounting : public FGameFeaturePluginState
 
 		if (MaybeRequestInfo.HasError())
 		{
-			ensureMsgf(false, TEXT("Unable to enqueue unmount for the PluginURL(%s) because %s"), *StateProperties.PluginIdentifier.GetFullPluginURL(), LexToString(MaybeRequestInfo.GetError()));
+			const FStringView ShortUrl = StateProperties.PluginIdentifier.GetIdentifyingString();
+			ensureMsgf(false, TEXT("Unable to enqueue unmount for the PluginURL(%.*s) because %s"), ShortUrl.Len(), ShortUrl.GetData(), LexToString(MaybeRequestInfo.GetError()));
 			Result = GetErrorResult(TEXT("BundleManager.Begin."), MaybeRequestInfo.GetError());
 			return;
 		}
@@ -1446,7 +1450,8 @@ struct FGameFeaturePluginState_Unmounting : public FGameFeaturePluginState
 
 		if (EnumHasAnyFlags(RequestInfo.InfoFlags, EInstallBundleRequestInfoFlags::SkippedUnknownBundles))
 		{
-			ensureMsgf(false, TEXT("Unable to enqueue unmount for the PluginURL(%s) because failed to resolve install bundles!"), *StateProperties.PluginIdentifier.GetFullPluginURL());
+			const FStringView ShortUrl = StateProperties.PluginIdentifier.GetIdentifyingString();
+			ensureMsgf(false, TEXT("Unable to enqueue unmount for the PluginURL(%.*s) because failed to resolve install bundles!"), ShortUrl.Len(), ShortUrl.GetData());
 			Result = GetErrorResult(TEXT("BundleManager.Begin."), TEXT("Cannot_Resolve"), UE::GameFeatures::CommonErrorCodes::Generic_ConnectionError);
 			return;
 		}
@@ -1533,7 +1538,8 @@ struct FGameFeaturePluginState_Mounting : public FGameFeaturePluginState
 	{
 		if (FPakFile* Pak = (FPakFile*)(&PakFile))
 		{
-			UE_LOG(LogGameFeatures, Display, TEXT("Mounted Pak File for (%s) with following files:"), *StateProperties.PluginIdentifier.GetFullPluginURL());
+			const FStringView ShortUrl = StateProperties.PluginIdentifier.GetIdentifyingString();
+			UE_LOG(LogGameFeatures, Display, TEXT("Mounted Pak File for (%.*s) with following files:"), ShortUrl.Len(), ShortUrl.GetData());
 			TArray<FString> OutFileList;
 			Pak->GetPrunedFilenames(OutFileList);
 			for (const FString& FileName : OutFileList)
@@ -1570,7 +1576,8 @@ struct FGameFeaturePluginState_Mounting : public FGameFeaturePluginState
 
 		if (MaybeRequestInfo.HasError())
 		{
-			ensureMsgf(false, TEXT("Unable to enqueue mount for the PluginURL(%s) because %s"), *StateProperties.PluginIdentifier.GetFullPluginURL(), LexToString(MaybeRequestInfo.GetError()));
+			const FStringView ShortUrl = StateProperties.PluginIdentifier.GetIdentifyingString();
+			ensureMsgf(false, TEXT("Unable to enqueue mount for the PluginURL(%.*s) because %s"), ShortUrl.Len(), ShortUrl.GetData(), LexToString(MaybeRequestInfo.GetError()));
 			Result = GetErrorResult(TEXT("BundleManager.Begin.CannotStart."), MaybeRequestInfo.GetError());
 			return;
 		}
@@ -1579,7 +1586,8 @@ struct FGameFeaturePluginState_Mounting : public FGameFeaturePluginState
 
 		if (EnumHasAnyFlags(RequestInfo.InfoFlags, EInstallBundleRequestInfoFlags::SkippedUnknownBundles))
 		{
-			ensureMsgf(false, TEXT("Unable to enqueue mount for the PluginURL(%s) because failed to resolve install bundles!"), *StateProperties.PluginIdentifier.GetFullPluginURL());
+			const FStringView ShortUrl = StateProperties.PluginIdentifier.GetIdentifyingString();
+			ensureMsgf(false, TEXT("Unable to enqueue mount for the PluginURL(%.*s) because failed to resolve install bundles!"), ShortUrl.Len(), ShortUrl.GetData());
 			Result = GetErrorResult(TEXT("BundleManager.Begin."), TEXT("Resolve_Failed"));
 			return;
 		}
@@ -2788,46 +2796,23 @@ EGameFeaturePluginProtocol FGameFeaturePluginStateMachineProperties::GetPluginPr
 	return PluginIdentifier.GetPluginProtocol();
 }
 
-FInstallBundlePluginProtocolMetaData::FInstallBundlePluginProtocolMetaData()
-{
-	ResetToDefaults();
-}
-
-void FInstallBundlePluginProtocolMetaData::ResetToDefaults()
-{
-	InstallBundles.Reset();
-	VersionNum = FDefaultValues::CurrentVersionNum;
-
-	static_assert(static_cast<uint8>(EGameFeatureInstallBundleProtocolOptions::Count) == 2, "Update this function to handle the newly added EGameFeatureInstallBundleProtocolOptions value!");
-}
-
 FString FInstallBundlePluginProtocolMetaData::ToString() const
 {
 	FString ReturnedString;
 
 	//Always encode InstallBundles
-	if (InstallBundles.Num() > 0)
-	{
-		ReturnedString = LexToString(EGameFeatureInstallBundleProtocolOptions::Bundles) + UE::GameFeatures::PluginURLStructureInfo::OptionAssignOperator;
+	ReturnedString = LexToString(EGameFeatureInstallBundleProtocolOptions::Bundles) + UE::GameFeatures::PluginURLStructureInfo::OptionAssignOperator;
 
-		FNameBuilder NameBuilder;
-		const FString BundlesList = FString::JoinBy(InstallBundles, UE::GameFeatures::PluginURLStructureInfo::OptionListSeperator, 
-		[&NameBuilder](const FName& BundleName)
-		{
-			BundleName.ToString(NameBuilder);
-			return FStringView(NameBuilder);
-		});
-		ReturnedString.Append(BundlesList);
-	}
-
-	//Always encode version
-	ReturnedString.Append(FString::Printf(TEXT("%s%s%s%d"), 
-		UE::GameFeatures::PluginURLStructureInfo::OptionSeperator, 
-		*LexToString(EGameFeatureInstallBundleProtocolOptions::Version), 
-		UE::GameFeatures::PluginURLStructureInfo::OptionAssignOperator, 
-		VersionNum));
+	FNameBuilder NameBuilder;
+	const FString BundlesList = FString::JoinBy(InstallBundles, UE::GameFeatures::PluginURLStructureInfo::OptionListSeperator, 
+	[&NameBuilder](const FName& BundleName)
+	{
+		BundleName.ToString(NameBuilder);
+		return FStringView(NameBuilder);
+	});
+	ReturnedString.Append(BundlesList);
 
-	static_assert(static_cast<uint8>(EGameFeatureInstallBundleProtocolOptions::Count) == 2, "Update this function to handle the newly added EGameFeatureInstallBundleProtocolOptions value!");
+	static_assert(static_cast<uint8>(EGameFeatureInstallBundleProtocolOptions::Count) == 1, "Update this function to handle the newly added EGameFeatureInstallBundleProtocolOptions value!");
 
 	return ReturnedString;
 }
@@ -2835,10 +2820,11 @@ FString FInstallBundlePluginProtocolMetaData::ToString() const
 bool FInstallBundlePluginProtocolMetaData::FromString(const FString& URLString, FInstallBundlePluginProtocolMetaData& Metadata)
 {
 	bool bParseSuccess = true;
-	TArray<FString> URLOptions;
-	URLString.ParseIntoArray(URLOptions, UE::GameFeatures::PluginURLStructureInfo::OptionSeperator);
 
-	Metadata.ResetToDefaults();
+	TArray<FStringView> URLOptions;
+	UE::String::ParseTokens(URLString, UE::GameFeatures::PluginURLStructureInfo::OptionSeperator, URLOptions);
+
+	Metadata = {};
 
 	if (URLOptions.Num() > 0)
 	{
@@ -2847,15 +2833,15 @@ bool FInstallBundlePluginProtocolMetaData::FromString(const FString& URLString,
 		//Parse through our URLOptions. Start at ParsingIndex 1 as option 0 should always be the .uplugin path that doesn't contain metadata
 		for (int ParsingIndex = 1; ParsingIndex < URLOptions.Num(); ++ParsingIndex)
 		{
-			const FString& URLOption = URLOptions[ParsingIndex];
+			const FStringView URLOption = URLOptions[ParsingIndex];
 	
-			TArray<FString> OptionStrings;
-			URLOption.ParseIntoArray(OptionStrings, UE::GameFeatures::PluginURLStructureInfo::OptionAssignOperator);
+			TArray<FStringView> OptionStrings;
+			UE::String::ParseTokens(URLOption, UE::GameFeatures::PluginURLStructureInfo::OptionAssignOperator, OptionStrings);
 
 			if (OptionStrings.Num() > 0)
 			{
 				EGameFeatureInstallBundleProtocolOptions OptionEnum = EGameFeatureInstallBundleProtocolOptions::Count;
-				LexFromString(OptionEnum, *(OptionStrings[0]));
+				LexFromString(OptionEnum, OptionStrings[0]);
 
 				switch (OptionEnum)
 				{
@@ -2868,43 +2854,29 @@ bool FInstallBundlePluginProtocolMetaData::FromString(const FString& URLString,
 						}
 						else
 						{
-							TArray<FString> BundleNames;
-							OptionStrings[1].ParseIntoArray(BundleNames, TEXT(","));
+							TArray<FStringView> BundleNames;
+							UE::String::ParseTokens(OptionStrings[1], UE::GameFeatures::PluginURLStructureInfo::OptionListSeperator, BundleNames);
 
 							Metadata.InstallBundles.Reserve(BundleNames.Num());
-							for (FString& BundleNameString : BundleNames)
+							for (const FStringView BundleNameString : BundleNames)
 							{
-								Metadata.InstallBundles.Add(*BundleNameString);
+								Metadata.InstallBundles.Emplace(BundleNameString);
 							}
 						}
 
 						break;
 					}
 
-					case EGameFeatureInstallBundleProtocolOptions::Version:
+					case EGameFeatureInstallBundleProtocolOptions::Count:
 					{
-						if (OptionStrings.Num() != 2)
-						{
-							bParseSuccess = false;
-							UE_LOG(LogGameFeatures, Warning, TEXT("Error parsing InstallBundle protocol options URL %s . Invalid Version Option!"), *URLString);
-						}
-						else
-						{
-							Metadata.VersionNum = FCString::Strtoi(*OptionStrings[1], nullptr, 10);
-							if (Metadata.VersionNum != FDefaultValues::CurrentVersionNum)
-							{
-								ensureAlwaysMsgf(false, TEXT("Mismatched version number for URL %s ! URL was generated on a previous version and might be missing data!"));
-							}
-						}
-
+						// Ignore unknown options as they may be game specfic
 						break;
 					}
 
-					case EGameFeatureInstallBundleProtocolOptions::Count:
 					default:
 					{
 						bParseSuccess = false;
-						UE_LOG(LogGameFeatures, Error, TEXT("Error parsing InstallBundle protocol options for URL %s . Unknown Option %s"), *URLString, *URLOption);
+						UE_LOG(LogGameFeatures, Error, TEXT("Error parsing InstallBundle protocol options for URL %s . Unknown Option %.*s"), *URLString, URLOption.Len(), URLOption.GetData());
 
 						break;
 					}
@@ -2919,7 +2891,7 @@ bool FInstallBundlePluginProtocolMetaData::FromString(const FString& URLString,
 		UE_LOG(LogGameFeatures, Error, TEXT("Error parsing InstallBundle protocol options URL %s . No Bundle List Found!"), *URLString);
 	}
 
-	static_assert(static_cast<uint8>(EGameFeatureInstallBundleProtocolOptions::Count) == 2, "Update this function to handle the newly added EGameFeatureInstallBundleProtocolOptions value!");
+	static_assert(static_cast<uint8>(EGameFeatureInstallBundleProtocolOptions::Count) == 1, "Update this function to handle the newly added EGameFeatureInstallBundleProtocolOptions value!");
 
 	return bParseSuccess;
 }
@@ -2984,10 +2956,12 @@ UE::GameFeatures::FResult FGameFeaturePluginStateMachineProperties::ValidateProt
 {
 	if (GetPluginProtocol() == EGameFeaturePluginProtocol::InstallBundle)
 	{
+		const FStringView ShortUrl = PluginIdentifier.GetIdentifyingString();
+
 		//Should never change our PluginProtocol
 		if (!ensureAlwaysMsgf(NewProtocolOptions.HasSubtype<FInstallBundlePluginProtocolOptions>()
-								,TEXT("Error with InstallBundle protocol FGameFeaturePluginStateMachineProperties having an invalid ProtocolOptions. URL: %s")
-								,*PluginIdentifier.GetFullPluginURL()))
+								,TEXT("Error with InstallBundle protocol FGameFeaturePluginStateMachineProperties having an invalid ProtocolOptions. URL: %.*s")
+								,ShortUrl.Len(), ShortUrl.GetData()))
 		{
 			return MakeError(UE::GameFeatures::StateMachineErrorNamespace + TEXT("ProtocolOptions.Invalid_Protocol"));
 		}
@@ -2997,7 +2971,7 @@ UE::GameFeatures::FResult FGameFeaturePluginStateMachineProperties::ValidateProt
 			const FInstallBundlePluginProtocolOptions& OldOptions = ProtocolOptions.GetSubtype<FInstallBundlePluginProtocolOptions>();
 			const FInstallBundlePluginProtocolOptions& NewOptions = NewProtocolOptions.GetSubtype<FInstallBundlePluginProtocolOptions>();
 
-			if (!ensureMsgf(OldOptions.bAllowIniLoading == NewOptions.bAllowIniLoading, TEXT("Unexpected change to AllowIniLoading when updating ProtocolOptions. URL: %s "), *PluginIdentifier.GetFullPluginURL()))
+			if (!ensureMsgf(OldOptions.bAllowIniLoading == NewOptions.bAllowIniLoading, TEXT("Unexpected change to AllowIniLoading when updating ProtocolOptions. URL: %.*s "), ShortUrl.Len(), ShortUrl.GetData()))
 			{
 				return MakeError(UE::GameFeatures::StateMachineErrorNamespace + TEXT("ProtocolOptions.Invalid_Update"));
 			}
diff --git a/Engine/Plugins/Experimental/GameFeatures/Source/GameFeatures/Private/GameFeaturePluginStateMachine.h b/Engine/Plugins/Experimental/GameFeatures/Source/GameFeatures/Private/GameFeaturePluginStateMachine.h
--- a/Engine/Plugins/Experimental/GameFeatures/Source/GameFeatures/Private/GameFeaturePluginStateMachine.h
+++ b/Engine/Plugins/Experimental/GameFeatures/Source/GameFeatures/Private/GameFeaturePluginStateMachine.h
@@ -218,27 +218,9 @@ struct FInstallBundlePluginProtocolMetaData
 {
 	TArray<FName> InstallBundles;
 
-	// TODO: keep this version?
-	/** Set to whatever the FDefaultValues::CurrentVersionNum was when this URL was generated.
-		Allows us to ensure if we try and load URLs generated from a previous version **/
-	uint8 VersionNum;
-
 	/** Functions to convert to/from the URL FString representation of this metadata **/
 	FString ToString() const;
 	static bool FromString(const FString& URLString, FInstallBundlePluginProtocolMetaData& OutMetadata);
-
-	FInstallBundlePluginProtocolMetaData();
-
-private:
-	/** Resets all our Metadata values to the default values */
-	void ResetToDefaults();
-
-	/** Holds default values for the above settings as they are only encoded into a string if they differ from these values */
-	struct FDefaultValues
-	{
-		static const uint32 CurrentVersionNum;
-		//Missing InstallBundles on purpose as the default is just an empty TArray and should always be encoded
-	};
 };
 
 struct FGameFeatureProtocolMetadata : public TUnion<FInstallBundlePluginProtocolMetaData, FNull>
diff --git a/Engine/Plugins/Experimental/GameFeatures/Source/GameFeatures/Private/GameFeaturesSubsystem.cpp b/Engine/Plugins/Experimental/GameFeatures/Source/GameFeatures/Private/GameFeaturesSubsystem.cpp
--- a/Engine/Plugins/Experimental/GameFeatures/Source/GameFeatures/Private/GameFeaturesSubsystem.cpp
+++ b/Engine/Plugins/Experimental/GameFeatures/Source/GameFeatures/Private/GameFeaturesSubsystem.cpp
@@ -603,10 +603,27 @@ void UGameFeaturesSubsystem::RemoveObserver(UObject* Observer)
 
 FString UGameFeaturesSubsystem::GetPluginURL_FileProtocol(const FString& PluginDescriptorPath)
 {
-	return TEXT("file:") + PluginDescriptorPath;
+	return UE::GameFeatures::GameFeaturePluginProtocolPrefix(EGameFeaturePluginProtocol::File) + PluginDescriptorPath;
 }
 
-FString GetPluginURL_InstallBundleProtocol(const FString& PluginName, const FInstallBundlePluginProtocolMetaData& ProtocolMetadata)
+FString UGameFeaturesSubsystem::GetPluginURL_FileProtocol(const FString& PluginDescriptorPath, TArrayView<const TPair<FString, FString>> AdditionalOptions)
+{
+	FString Path;
+	Path += UE::GameFeatures::GameFeaturePluginProtocolPrefix(EGameFeaturePluginProtocol::File);
+	Path += PluginDescriptorPath;
+	if (AdditionalOptions.Num() > 0)
+	{
+		Path += UE::GameFeatures::PluginURLStructureInfo::OptionSeperator;
+		Path += FString::JoinBy(AdditionalOptions, UE::GameFeatures::PluginURLStructureInfo::OptionSeperator,
+			[](const TPair<FString, FString>& OptionPair)
+			{
+				return OptionPair.Key + UE::GameFeatures::PluginURLStructureInfo::OptionAssignOperator + OptionPair.Value;
+			});
+	}
+	return Path;
+}
+
+FString GetPluginURL_InstallBundleProtocol(const FString& PluginName, const FInstallBundlePluginProtocolMetaData& ProtocolMetadata, TArrayView<const TPair<FString, FString>> AdditionalOptions = {})
 {
 	ensure(ProtocolMetadata.InstallBundles.Num() > 0);
 	FString Path;
@@ -614,6 +631,15 @@ FString GetPluginURL_InstallBundleProtocol(const FString& PluginName, const FIns
 	Path += PluginName;
 	Path += UE::GameFeatures::PluginURLStructureInfo::OptionSeperator;
 	Path += ProtocolMetadata.ToString();
+	if (AdditionalOptions.Num() > 0)
+	{
+		Path += UE::GameFeatures::PluginURLStructureInfo::OptionSeperator;
+		Path += FString::JoinBy(AdditionalOptions, UE::GameFeatures::PluginURLStructureInfo::OptionSeperator,
+			[](const TPair<FString, FString>& OptionPair)
+			{
+				return OptionPair.Key + UE::GameFeatures::PluginURLStructureInfo::OptionAssignOperator + OptionPair.Value;
+			});
+	}
 
 	return Path;
 }
@@ -623,7 +649,7 @@ FString UGameFeaturesSubsystem::GetPluginURL_InstallBundleProtocol(const FString
 	FInstallBundlePluginProtocolMetaData ProtocolMetadata;
 	for (const FString& BundleName : BundleNames)
 	{
-		ProtocolMetadata.InstallBundles.Add(FName(BundleName));
+		ProtocolMetadata.InstallBundles.Emplace(BundleName);
 	}
 	return ::GetPluginURL_InstallBundleProtocol(PluginName, ProtocolMetadata);
 }
@@ -635,9 +661,7 @@ FString UGameFeaturesSubsystem::GetPluginURL_InstallBundleProtocol(const FString
 
 FString UGameFeaturesSubsystem::GetPluginURL_InstallBundleProtocol(const FString& PluginName, const TArrayView<const FName> BundleNames)
 {
-	FInstallBundlePluginProtocolMetaData ProtocolMetadata;
-	ProtocolMetadata.InstallBundles.Append(BundleNames.GetData(), BundleNames.Num());
-	return ::GetPluginURL_InstallBundleProtocol(PluginName, ProtocolMetadata);
+	return GetPluginURL_InstallBundleProtocol(PluginName, BundleNames, TArrayView<const TPair<FString, FString>>());
 }
 
 FString UGameFeaturesSubsystem::GetPluginURL_InstallBundleProtocol(const FString& PluginName, FName BundleName)
@@ -645,6 +669,13 @@ FString UGameFeaturesSubsystem::GetPluginURL_InstallBundleProtocol(const FString
 	return GetPluginURL_InstallBundleProtocol(PluginName, MakeArrayView(&BundleName, 1));
 }
 
+FString UGameFeaturesSubsystem::GetPluginURL_InstallBundleProtocol(const FString& PluginName, TArrayView<const FName> BundleNames, TArrayView<const TPair<FString, FString>> AdditionalOptions)
+{
+	FInstallBundlePluginProtocolMetaData ProtocolMetadata;
+	ProtocolMetadata.InstallBundles.Append(BundleNames.GetData(), BundleNames.Num());
+	return ::GetPluginURL_InstallBundleProtocol(PluginName, ProtocolMetadata, AdditionalOptions);
+}
+
 EGameFeaturePluginProtocol UGameFeaturesSubsystem::GetPluginURLProtocol(FStringView PluginURL)
 {
 	for (EGameFeaturePluginProtocol Protocol : TEnumRange<EGameFeaturePluginProtocol>())
@@ -1639,18 +1670,20 @@ UGameFeaturePluginStateMachine* UGameFeaturesSubsystem::FindGameFeaturePluginSta
 
 UGameFeaturePluginStateMachine* UGameFeaturesSubsystem::FindGameFeaturePluginStateMachine(const FGameFeaturePluginIdentifier& PluginIdentifier) const
 {
+	const FStringView ShortUrl = PluginIdentifier.GetIdentifyingString();
+
 	TObjectPtr<UGameFeaturePluginStateMachine> const* ExistingStateMachine = 
 		GameFeaturePluginStateMachines.FindByHash(GetTypeHash(PluginIdentifier.GetIdentifyingString()), PluginIdentifier.GetIdentifyingString());
 	if (ExistingStateMachine)
 	{
 		EGameFeaturePluginProtocol ExpectedProtocol = (*ExistingStateMachine)->GetPluginIdentifier().GetPluginProtocol();
-		if (ensureMsgf(ExpectedProtocol == PluginIdentifier.GetPluginProtocol(), TEXT("Expected protocol %s for %s"), UE::GameFeatures::GameFeaturePluginProtocolPrefix(ExpectedProtocol), *PluginIdentifier.GetFullPluginURL()))
+		if (ensureMsgf(ExpectedProtocol == PluginIdentifier.GetPluginProtocol(), TEXT("Expected protocol %s for %.*s"), UE::GameFeatures::GameFeaturePluginProtocolPrefix(ExpectedProtocol), ShortUrl.Len(), ShortUrl.GetData()))
 		{
-			UE_LOG(LogGameFeatures, VeryVerbose, TEXT("FOUND GameFeaturePlugin using PluginIdentifier:%.*s for PluginURL:%s"), PluginIdentifier.GetIdentifyingString().Len(), PluginIdentifier.GetIdentifyingString().GetData(), *PluginIdentifier.GetFullPluginURL());
+			UE_LOG(LogGameFeatures, VeryVerbose, TEXT("FOUND GameFeaturePlugin using PluginIdentifier:%.*s for PluginURL:%s"), ShortUrl.Len(), ShortUrl.GetData(), *PluginIdentifier.GetFullPluginURL());
 			return *ExistingStateMachine;
 		}
 	}
-	UE_LOG(LogGameFeatures, VeryVerbose, TEXT("NOT FOUND GameFeaturePlugin using PluginIdentifier:%.*s for PluginURL:%s"), PluginIdentifier.GetIdentifyingString().Len(), PluginIdentifier.GetIdentifyingString().GetData(), *PluginIdentifier.GetFullPluginURL());
+	UE_LOG(LogGameFeatures, VeryVerbose, TEXT("NOT FOUND GameFeaturePlugin using PluginIdentifier:%.*s for PluginURL:%s"), ShortUrl.Len(), ShortUrl.GetData(), *PluginIdentifier.GetFullPluginURL());
 
 	return nullptr;
 }
@@ -1666,7 +1699,7 @@ UGameFeaturePluginStateMachine* UGameFeaturesSubsystem::FindOrCreateGameFeatureP
 	if (ExistingStateMachine)
 	{
 		EGameFeaturePluginProtocol ExpectedProtocol = (*ExistingStateMachine)->GetPluginIdentifier().GetPluginProtocol();
-		ensureAlwaysMsgf(ExpectedProtocol == PluginIdentifier.GetPluginProtocol(), TEXT("Expected protocol %s for %s"), UE::GameFeatures::GameFeaturePluginProtocolPrefix(ExpectedProtocol), *PluginIdentifier.GetFullPluginURL());
+		ensureAlwaysMsgf(ExpectedProtocol == PluginIdentifier.GetPluginProtocol(), TEXT("Expected protocol %s for %.*s"), UE::GameFeatures::GameFeaturePluginProtocolPrefix(ExpectedProtocol), PluginIdentifier.GetIdentifyingString().Len(), PluginIdentifier.GetIdentifyingString().GetData());
 
 		// In this case, still return the existing machine, even if the protocol doesn't match. This function should never return null.
 		// There can only be one active instance of any machine.
diff --git a/Engine/Plugins/Experimental/GameFeatures/Source/GameFeatures/Public/GameFeaturesSubsystem.h b/Engine/Plugins/Experimental/GameFeatures/Source/GameFeatures/Public/GameFeaturesSubsystem.h
--- a/Engine/Plugins/Experimental/GameFeatures/Source/GameFeatures/Public/GameFeaturesSubsystem.h
+++ b/Engine/Plugins/Experimental/GameFeatures/Source/GameFeatures/Public/GameFeaturesSubsystem.h
@@ -29,13 +29,13 @@ namespace UE::GameFeatures
 	namespace PluginURLStructureInfo
 	{
 		/** Character used to denote what value is being assigned to the option before it */
-		extern const TCHAR* OptionAssignOperator;
+		extern GAMEFEATURES_API const TCHAR* OptionAssignOperator;
 
 		/** Character used to separate options on the URL. Used between each assigned value and the next Option name. */
-		extern const TCHAR* OptionSeperator;
+		extern GAMEFEATURES_API const TCHAR* OptionSeperator;
 
 		/** Character used to separate lists of values for a single option. Used between each entry in the list. */
-		extern const TCHAR* OptionListSeperator;
+		extern GAMEFEATURES_API const TCHAR* OptionListSeperator;
 	};
 
 	namespace CommonErrorCodes
@@ -368,12 +368,14 @@ public:
 public:
 	/** Construct a 'file:' Plugin URL using from the PluginDescriptorPath */
 	static FString GetPluginURL_FileProtocol(const FString& PluginDescriptorPath);
+	static FString GetPluginURL_FileProtocol(const FString& PluginDescriptorPath, TArrayView<const TPair<FString, FString>> AdditionalOptions);
 
 	/** Construct a 'installbundle:' Plugin URL using from the PluginName and required install bundles */
 	static FString GetPluginURL_InstallBundleProtocol(const FString& PluginName, TArrayView<const FString> BundleNames);
 	static FString GetPluginURL_InstallBundleProtocol(const FString& PluginName, const FString& BundleName);
 	static FString GetPluginURL_InstallBundleProtocol(const FString& PluginName, TArrayView<const FName> BundleNames);
 	static FString GetPluginURL_InstallBundleProtocol(const FString& PluginName, FName BundleName);
+	static FString GetPluginURL_InstallBundleProtocol(const FString& PluginName, TArrayView<const FName> BundleNames, TArrayView<const TPair<FString, FString>> AdditionalOptions);
 
 	/** Returns the plugin protocol for the specified URL */
 	static EGameFeaturePluginProtocol GetPluginURLProtocol(FStringView PluginURL);
