Functions to transition multiple GFPs at the same time.
- Not including unload functions at this time because I think there are still questions around dependency handling

diff --git a/Engine/Plugins/Experimental/GameFeatures/Source/GameFeatures/Private/GameFeaturesSubsystem.cpp b/Engine/Plugins/Experimental/GameFeatures/Source/GameFeatures/Private/GameFeaturesSubsystem.cpp
--- a/Engine/Plugins/Experimental/GameFeatures/Source/GameFeatures/Private/GameFeaturesSubsystem.cpp
+++ b/Engine/Plugins/Experimental/GameFeatures/Source/GameFeatures/Private/GameFeaturesSubsystem.cpp
@@ -900,6 +900,47 @@ void UGameFeaturesSubsystem::LoadGameFeaturePlugin(const FString& PluginURL, con
 	ChangeGameFeatureDestination(StateMachine, ProtocolOptions, FGameFeaturePluginStateRange(EGameFeaturePluginState::Loaded, EGameFeaturePluginState::Active), CompleteDelegate);
 }
 
+void UGameFeaturesSubsystem::LoadGameFeaturePlugin(TConstArrayView<FString> PluginURLs, const FGameFeatureProtocolOptions& ProtocolOptions, const FMultipleGameFeaturePluginsLoaded& CompleteDelegate)
+{
+	struct FLoadContext
+	{
+		TMap<FString, UE::GameFeatures::FResult> Results;
+		FMultipleGameFeaturePluginsLoaded CompleteDelegate;
+
+		int32 NumPluginsLoaded = 0;
+
+		~FLoadContext()
+		{
+			UGameplayTagsManager::Get().PopDeferOnGameplayTagTreeChangedBroadcast();
+
+			CompleteDelegate.ExecuteIfBound(Results);
+		}
+	};
+	TSharedRef<FLoadContext> LoadContext = MakeShared<FLoadContext>();
+	LoadContext->CompleteDelegate = CompleteDelegate;
+
+	UGameplayTagsManager::Get().PushDeferOnGameplayTagTreeChangedBroadcast();
+
+	LoadContext->Results.Reserve(PluginURLs.Num());
+	for (const FString& PluginURL : PluginURLs)
+	{
+		LoadContext->Results.Add(PluginURL, MakeError("Pending"));
+	}
+
+	const int32 NumPluginsToLoad = PluginURLs.Num();
+	UE_LOG(LogGameFeatures, Log, TEXT("Loading %i GFPs"), NumPluginsToLoad);
+
+	for (const FString& PluginURL : PluginURLs)
+	{
+		LoadGameFeaturePlugin(PluginURL, ProtocolOptions, FGameFeaturePluginChangeStateComplete::CreateLambda([LoadContext, PluginURL](const UE::GameFeatures::FResult& Result)
+		{
+			LoadContext->Results.Add(PluginURL, Result);
+			++LoadContext->NumPluginsLoaded;
+			UE_LOG(LogGameFeatures, VeryVerbose, TEXT("Finished Loading %i GFPs"), LoadContext->NumPluginsLoaded);
+		}));
+	}
+}
+
 void UGameFeaturesSubsystem::LoadAndActivateGameFeaturePlugin(const FString& PluginURL, const FGameFeaturePluginLoadComplete& CompleteDelegate)
 {
 	ChangeGameFeatureTargetState(PluginURL, EGameFeatureTargetState::Active, CompleteDelegate);
@@ -910,6 +951,11 @@ void UGameFeaturesSubsystem::LoadAndActivateGameFeaturePlugin(const FString& Plu
 	ChangeGameFeatureTargetState(PluginURL, ProtocolOptions, EGameFeatureTargetState::Active, CompleteDelegate);
 }
 
+void UGameFeaturesSubsystem::LoadAndActivateGameFeaturePlugin(TConstArrayView<FString> PluginURLs, const FGameFeatureProtocolOptions& ProtocolOptions, const FMultipleGameFeaturePluginsLoaded& CompleteDelegate)
+{
+	ChangeGameFeatureTargetState(PluginURLs, ProtocolOptions, EGameFeatureTargetState::Active, CompleteDelegate);
+}
+
 void UGameFeaturesSubsystem::ChangeGameFeatureTargetState(const FString& PluginURL, EGameFeatureTargetState TargetState, const FGameFeaturePluginChangeStateComplete& CompleteDelegate)
 {
 	ChangeGameFeatureTargetState(PluginURL, FGameFeatureProtocolOptions(), TargetState, CompleteDelegate);
@@ -988,6 +1034,47 @@ void UGameFeaturesSubsystem::ChangeGameFeatureTargetState(const FString& PluginU
 	ChangeGameFeatureDestination(StateMachine, ProtocolOptions, FGameFeaturePluginStateRange(TargetPluginState), CompleteDelegate);
 }
 
+void UGameFeaturesSubsystem::ChangeGameFeatureTargetState(TConstArrayView<FString> PluginURLs, const FGameFeatureProtocolOptions& ProtocolOptions, EGameFeatureTargetState TargetState, const FMultipleGameFeaturePluginsLoaded& CompleteDelegate)
+{
+	struct FLoadContext
+	{
+		TMap<FString, UE::GameFeatures::FResult> Results;
+		FMultipleGameFeaturePluginsLoaded CompleteDelegate;
+
+		int32 NumPluginsLoaded = 0;
+
+		~FLoadContext()
+		{
+			UGameplayTagsManager::Get().PopDeferOnGameplayTagTreeChangedBroadcast();
+
+			CompleteDelegate.ExecuteIfBound(Results);
+		}
+	};
+	TSharedRef<FLoadContext> LoadContext = MakeShared<FLoadContext>();
+	LoadContext->CompleteDelegate = CompleteDelegate;
+
+	UGameplayTagsManager::Get().PushDeferOnGameplayTagTreeChangedBroadcast();
+
+	LoadContext->Results.Reserve(PluginURLs.Num());
+	for (const FString& PluginURL : PluginURLs)
+	{
+		LoadContext->Results.Add(PluginURL, MakeError("Pending"));
+	}
+
+	const int32 NumPluginsToLoad = PluginURLs.Num();
+	UE_LOG(LogGameFeatures, Log, TEXT("Transitioning (%s) %i GFPs"), *LexToString(TargetState), NumPluginsToLoad);
+
+	for (const FString& PluginURL : PluginURLs)
+	{
+		ChangeGameFeatureTargetState(PluginURL, ProtocolOptions, TargetState, FGameFeaturePluginChangeStateComplete::CreateLambda([LoadContext, PluginURL, TargetState](const UE::GameFeatures::FResult& Result)
+		{
+			LoadContext->Results.Add(PluginURL, Result);
+			++LoadContext->NumPluginsLoaded;
+			UE_LOG(LogGameFeatures, VeryVerbose, TEXT("Finished Transitioning (%s) %i GFPs"), *LexToString(TargetState), LoadContext->NumPluginsLoaded);
+		}));
+	}
+}
+
 UE::GameFeatures::FResult UGameFeaturesSubsystem::UpdateGameFeatureProtocolOptions(const FString& PluginURL, const FGameFeatureProtocolOptions& NewOptions, bool* bOutDidUpdate /*= nullptr*/)
 {
 	UGameFeaturePluginStateMachine* StateMachine = FindGameFeaturePluginStateMachine(PluginURL);
@@ -1067,6 +1154,30 @@ bool UGameFeaturesSubsystem::GetGameFeaturePluginInstallPercent(const FString& P
 	return false;
 }
 
+bool UGameFeaturesSubsystem::GetGameFeaturePluginInstallPercent(TConstArrayView<FString> PluginURLs, float& Install_Percent) const
+{
+	float TotalInstallPercent = 0;
+	int32 NumFound = 0;
+
+	for (const FString& URL : PluginURLs)
+	{
+		float SingleInstallPercent = 0;
+		if (GetGameFeaturePluginInstallPercent(URL, SingleInstallPercent))
+		{
+			TotalInstallPercent += SingleInstallPercent;
+			++NumFound;
+		}
+	}
+
+	if (NumFound > 0)
+	{
+		Install_Percent = TotalInstallPercent / NumFound;
+		return true;
+	}
+
+	return false;
+}
+
 bool UGameFeaturesSubsystem::IsGameFeaturePluginActive(const FString& PluginURL, bool bCheckForActivating /*= false*/) const
 {
 	if (const UGameFeaturePluginStateMachine* StateMachine = FindGameFeaturePluginStateMachine(PluginURL))
@@ -1242,6 +1353,42 @@ void UGameFeaturesSubsystem::CancelGameFeatureStateChange(const FString& PluginU
 	}
 }
 
+void UGameFeaturesSubsystem::CancelGameFeatureStateChange(TConstArrayView<FString> PluginURLs, const FMultipleGameFeaturePluginChangeStateComplete& CompleteDelegate)
+{
+	struct FContext
+	{
+		TMap<FString, UE::GameFeatures::FResult> Results;
+		FMultipleGameFeaturePluginsLoaded CompleteDelegate;
+
+		int32 NumPluginsCanceled = 0;
+
+		~FContext()
+		{
+			CompleteDelegate.ExecuteIfBound(Results);
+		}
+	};
+	TSharedRef<FContext> CancelContext = MakeShared<FContext>();
+	CancelContext->CompleteDelegate = CompleteDelegate;
+
+	CancelContext->Results.Reserve(PluginURLs.Num());
+	for (const FString& PluginURL : PluginURLs)
+	{
+		CancelContext->Results.Add(PluginURL, MakeError("Pending"));
+	}
+
+	UE_LOG(LogGameFeatures, Log, TEXT("Canceling %i GFP transitions"), PluginURLs.Num());
+
+	for (const FString& PluginURL : PluginURLs)
+	{
+		CancelGameFeatureStateChange(PluginURL, FGameFeaturePluginChangeStateComplete::CreateLambda([CancelContext, PluginURL](const UE::GameFeatures::FResult& Result)
+		{
+			CancelContext->Results.Add(PluginURL, Result);
+			++CancelContext->NumPluginsCanceled;
+			UE_LOG(LogGameFeatures, VeryVerbose, TEXT("Finished canceling %i GFP transitions"), CancelContext->NumPluginsCanceled);
+		}));
+	}
+}
+
 void UGameFeaturesSubsystem::LoadBuiltInGameFeaturePlugin(const TSharedRef<IPlugin>& Plugin, FBuiltInPluginAdditionalFilters AdditionalFilter, const FGameFeaturePluginLoadComplete& CompleteDelegate /*= FGameFeaturePluginLoadComplete()*/)
 {
 	UE_SCOPED_ENGINE_ACTIVITY(TEXT("Loading GameFeaturePlugin %s"), *Plugin->GetName());
@@ -1382,7 +1529,7 @@ void UGameFeaturesSubsystem::LoadBuiltInGameFeaturePlugins(FBuiltInPluginAdditio
 			CompleteDelegate.ExecuteIfBound(Results);
 		}
 	};
-	TSharedRef LoadContext = MakeShared<FLoadContext>();
+	TSharedRef<FLoadContext> LoadContext = MakeShared<FLoadContext>();
 	LoadContext->CompleteDelegate = InCompleteDelegate;
 
 	UAssetManager::Get().PushBulkScanning();
@@ -1403,7 +1550,7 @@ void UGameFeaturesSubsystem::LoadBuiltInGameFeaturePlugins(FBuiltInPluginAdditio
 	for (const TSharedRef<IPlugin>& Plugin : EnabledPlugins)
 	{
 		FBuiltInPluginLoadTimeTrackerScope TrackerScope(PluginLoadTimeTracker, Plugin);
-		LoadBuiltInGameFeaturePlugin(Plugin, AdditionalFilter, FGameFeaturePluginLoadComplete::CreateLambda([LoadContext, Plugin, NumPluginsToLoad](const UE::GameFeatures::FResult& Result)
+		LoadBuiltInGameFeaturePlugin(Plugin, AdditionalFilter, FGameFeaturePluginLoadComplete::CreateLambda([LoadContext, Plugin](const UE::GameFeatures::FResult& Result)
 		{
 			LoadContext->Results.Add(Plugin->GetName(), Result);
 			++LoadContext->NumPluginsLoaded;
diff --git a/Engine/Plugins/Experimental/GameFeatures/Source/GameFeatures/Public/GameFeaturesSubsystem.h b/Engine/Plugins/Experimental/GameFeatures/Source/GameFeatures/Public/GameFeaturesSubsystem.h
--- a/Engine/Plugins/Experimental/GameFeatures/Source/GameFeatures/Public/GameFeaturesSubsystem.h
+++ b/Engine/Plugins/Experimental/GameFeatures/Source/GameFeatures/Public/GameFeaturesSubsystem.h
@@ -167,7 +167,10 @@ using FGameFeaturePluginUninstallComplete = FGameFeaturePluginChangeStateComplet
 using FGameFeaturePluginTerminateComplete = FGameFeaturePluginChangeStateComplete;
 using FGameFeaturePluginUpdateProtocolComplete = FGameFeaturePluginChangeStateComplete;
 
-DECLARE_DELEGATE_OneParam(FBuiltInGameFeaturePluginsLoaded, PREPROCESSOR_COMMA_SEPARATED(const TMap<FString, UE::GameFeatures::FResult>& /*Results*/));
+DECLARE_DELEGATE_OneParam(FMultipleGameFeaturePluginChangeStateComplete, PREPROCESSOR_COMMA_SEPARATED(const TMap<FString, UE::GameFeatures::FResult>& /*Results*/));
+
+using FBuiltInGameFeaturePluginsLoaded = FMultipleGameFeaturePluginChangeStateComplete;
+using FMultipleGameFeaturePluginsLoaded = FMultipleGameFeaturePluginChangeStateComplete;
 
 enum class EBuiltInAutoState : uint8
 {
@@ -405,20 +408,24 @@ public:
 	/** Loads a single game feature plugin. */
 	void LoadGameFeaturePlugin(const FString& PluginURL, const FGameFeaturePluginLoadComplete& CompleteDelegate);
 	void LoadGameFeaturePlugin(const FString& PluginURL, const FGameFeatureProtocolOptions& ProtocolOptions, const FGameFeaturePluginLoadComplete& CompleteDelegate);
+	void LoadGameFeaturePlugin(TConstArrayView<FString> PluginURLs, const FGameFeatureProtocolOptions& ProtocolOptions, const FMultipleGameFeaturePluginsLoaded& CompleteDelegate);
 
 	/** Loads a single game feature plugin and activates it. */
 	void LoadAndActivateGameFeaturePlugin(const FString& PluginURL, const FGameFeaturePluginLoadComplete& CompleteDelegate);
 	void LoadAndActivateGameFeaturePlugin(const FString& PluginURL, const FGameFeatureProtocolOptions& ProtocolOptions, const FGameFeaturePluginLoadComplete& CompleteDelegate);
+	void LoadAndActivateGameFeaturePlugin(TConstArrayView<FString> PluginURLs, const FGameFeatureProtocolOptions& ProtocolOptions, const FMultipleGameFeaturePluginsLoaded& CompleteDelegate);
 
 	/** Changes the target state of a game feature plugin */
 	void ChangeGameFeatureTargetState(const FString& PluginURL, EGameFeatureTargetState TargetState, const FGameFeaturePluginChangeStateComplete& CompleteDelegate);
 	void ChangeGameFeatureTargetState(const FString& PluginURL, const FGameFeatureProtocolOptions& ProtocolOptions, EGameFeatureTargetState TargetState, const FGameFeaturePluginChangeStateComplete& CompleteDelegate);
+	void ChangeGameFeatureTargetState(TConstArrayView<FString> PluginURLs, const FGameFeatureProtocolOptions& ProtocolOptions, EGameFeatureTargetState TargetState, const FMultipleGameFeaturePluginsLoaded& CompleteDelegate);
 
 	/** Changes the protocol options of a game feature plugin. Useful to change any options data such as settings flags */
 	UE::GameFeatures::FResult UpdateGameFeatureProtocolOptions(const FString& PluginURL, const FGameFeatureProtocolOptions& NewOptions, bool* bOutDidUpdate = nullptr);
 
 	/** Gets the Install_Percent for single game feature plugin if it is active. */
 	bool GetGameFeaturePluginInstallPercent(const FString& PluginURL, float& Install_Percent) const;
+	bool GetGameFeaturePluginInstallPercent(TConstArrayView<FString> PluginURLs, float& Install_Percent) const;
 
 	/** Determines if a plugin is in the Active state.*/
 	bool IsGameFeaturePluginActive(const FString& PluginURL, bool bCheckForActivating = false) const;
@@ -449,6 +456,7 @@ public:
 	/** Attempt to cancel any state change. Calls back when cancelation is complete. Any other pending callbacks will be called with a canceled error. */
 	void CancelGameFeatureStateChange(const FString& PluginURL);
 	void CancelGameFeatureStateChange(const FString& PluginURL, const FGameFeaturePluginChangeStateComplete& CompleteDelegate);
+	void CancelGameFeatureStateChange(TConstArrayView<FString> PluginURLs, const FMultipleGameFeaturePluginChangeStateComplete& CompleteDelegate);
 
 	/**
 	 * If the specified plugin is known by the game feature system, returns the URL used to identify it
