#UE Add virtual to allow GameFeatureData to procedurally disallow scanning for primary assets

diff --git a/Engine/Plugins/Experimental/GameFeatures/Source/GameFeatures/Public/GameFeatureData.h b/Engine/Plugins/Experimental/GameFeatures/Source/GameFeatures/Public/GameFeatureData.h
--- a/Engine/Plugins/Experimental/GameFeatures/Source/GameFeatures/Public/GameFeatureData.h
+++ b/Engine/Plugins/Experimental/GameFeatures/Source/GameFeatures/Public/GameFeatureData.h
@@ -20,10 +20,10 @@ class GAMEFEATURES_API UGameFeatureData : public UPrimaryDataAsset
 
 public:
 	/** Method to get where the primary assets should scanned from in the plugin hierarchy */
-	const TArray<FPrimaryAssetTypeInfo>& GetPrimaryAssetTypesToScan() const { return PrimaryAssetTypesToScan; }
+	virtual const TArray<FPrimaryAssetTypeInfo>& GetPrimaryAssetTypesToScan() const { return PrimaryAssetTypesToScan; }
 
 #if WITH_EDITOR
-	TArray<FPrimaryAssetTypeInfo>& GetPrimaryAssetTypesToScan() { return PrimaryAssetTypesToScan; }
+	virtual TArray<FPrimaryAssetTypeInfo>& GetPrimaryAssetTypesToScan() { return PrimaryAssetTypesToScan; }
 #endif //if WITH_EDITOR
 
 	/** Method to process the base ini file for the plugin during loading */
