Fix for game feature data details customization crash when entering PIE with a game feature data plugin window visible.
#rnx

diff --git a/Engine/Plugins/Experimental/GameFeatures/Source/GameFeaturesEditor/Private/GameFeatureDataDetailsCustomization.cpp b/Engine/Plugins/Experimental/GameFeatures/Source/GameFeaturesEditor/Private/GameFeatureDataDetailsCustomization.cpp
--- a/Engine/Plugins/Experimental/GameFeatures/Source/GameFeaturesEditor/Private/GameFeatureDataDetailsCustomization.cpp
+++ b/Engine/Plugins/Experimental/GameFeatures/Source/GameFeaturesEditor/Private/GameFeatureDataDetailsCustomization.cpp
@@ -132,7 +132,7 @@ void FGameFeatureDataDetailsCustomization::CustomizeDetails(IDetailLayoutBuilder
 				.Padding(0.0f, 4.0f, 0.0f, 0.0f)
 				[
 					SNew(SHorizontalBox)
-					.Visibility_Lambda([this]() { return (GetCurrentState() == EGameFeaturePluginState::Active) ? EVisibility::Visible : EVisibility::Collapsed; })
+					.Visibility(this, &FGameFeatureDataDetailsCustomization::GetVisbililty)
 					+SHorizontalBox::Slot()
 					.AutoWidth()
 					.Padding(Padding)
@@ -217,6 +217,11 @@ EGameFeaturePluginState FGameFeatureDataDetailsCustomization::GetCurrentState()
 	return UGameFeaturesSubsystem::Get().GetPluginState(PluginURL);
 }
 
+EVisibility FGameFeatureDataDetailsCustomization::GetVisbililty() const
+{
+	return (GetCurrentState() == EGameFeaturePluginState::Active) ? EVisibility::Visible : EVisibility::Collapsed;
+}
+
 FText FGameFeatureDataDetailsCustomization::GetInitialStateText() const
 {
 	const EBuiltInAutoState AutoState = UGameFeaturesSubsystem::DetermineBuiltInInitialFeatureState(PluginPtr->GetDescriptor().CachedJson, FString());
diff --git a/Engine/Plugins/Experimental/GameFeatures/Source/GameFeaturesEditor/Private/GameFeatureDataDetailsCustomization.h b/Engine/Plugins/Experimental/GameFeatures/Source/GameFeaturesEditor/Private/GameFeatureDataDetailsCustomization.h
--- a/Engine/Plugins/Experimental/GameFeatures/Source/GameFeaturesEditor/Private/GameFeatureDataDetailsCustomization.h
+++ b/Engine/Plugins/Experimental/GameFeatures/Source/GameFeaturesEditor/Private/GameFeatureDataDetailsCustomization.h
@@ -11,6 +11,7 @@ namespace UE::GameFeatures { struct FResult; }
 class IDetailLayoutBuilder;
 class SErrorText;
 class IPlugin;
+struct EVisibility;
 
 //////////////////////////////////////////////////////////////////////////
 // FGameFeatureDataDetailsCustomization
@@ -29,6 +30,7 @@ protected:
 	void ChangeDesiredState(EGameFeaturePluginState State);
 
 	EGameFeaturePluginState GetCurrentState() const;
+	EVisibility GetVisbililty() const;
 
 	FText GetInitialStateText() const;
 	FText GetTagConfigPathText() const;
