Added the ability to selectively activate listed plugin dependencies when a plugin is activated.
- Requires `Activate=true` on a dependency in the .uplugin file.

New States
- Added ActivatingDependencies state between Loaded and Activated that will selectively activate required dependencies and wait for completetion.
- Added ErrorActivatingDependencies to handle errors in the ActivatingDependencies state. Can transtion on to Unloading or ActivatingDependencies.

diff --git a/Engine/Plugins/Experimental/GameFeatures/Source/GameFeatures/Private/GameFeaturePluginStateMachine.cpp b/Engine/Plugins/Experimental/GameFeatures/Source/GameFeatures/Private/GameFeaturePluginStateMachine.cpp
--- a/Engine/Plugins/Experimental/GameFeatures/Source/GameFeatures/Private/GameFeaturePluginStateMachine.cpp
+++ b/Engine/Plugins/Experimental/GameFeatures/Source/GameFeatures/Private/GameFeaturePluginStateMachine.cpp
@@ -480,6 +480,185 @@
 =========================================================
 */
 
+template<typename TransitionPolicy>
+struct FTransitionDependenciesGameFeaturePluginState : public FGameFeaturePluginState
+{
+	FTransitionDependenciesGameFeaturePluginState(FGameFeaturePluginStateMachineProperties& InStateProperties)
+		: FGameFeaturePluginState(InStateProperties)
+		, bRequestedDependencies(false)
+	{}
+
+	virtual ~FTransitionDependenciesGameFeaturePluginState()
+	{
+		ClearDependencies();
+	}
+
+	virtual void BeginState() override
+	{
+		ClearDependencies();
+	}
+
+	virtual void EndState() override
+	{
+		ClearDependencies();
+	}
+
+	virtual void UpdateState(FGameFeaturePluginStateStatus& StateStatus) override
+	{
+		TRACE_CPUPROFILER_EVENT_SCOPE(GFP_TransitionDependencies);
+		checkf(!StateProperties.PluginInstalledFilename.IsEmpty(), TEXT("PluginInstalledFilename must be set by the loading dependencies phase. PluginURL: %s"), *StateProperties.PluginIdentifier.GetFullPluginURL());
+		checkf(FPaths::GetExtension(StateProperties.PluginInstalledFilename) == TEXT("uplugin"), TEXT("PluginInstalledFilename must have a uplugin extension. PluginURL: %s"), *StateProperties.PluginIdentifier.GetFullPluginURL());
+
+		if (!bRequestedDependencies)
+		{
+			UGameFeaturesSubsystem& GameFeaturesSubsystem = UGameFeaturesSubsystem::Get();
+
+			TArray<UGameFeaturePluginStateMachine*> Dependencies;
+			if (!TransitionPolicy::GetPluginDependencyStateMachines(StateProperties, Dependencies))
+			{
+				// Failed to query dependencies
+				StateStatus.SetTransitionError(TransitionPolicy::GetErrorState(), GetErrorResult(TEXT("Failed_Dependency_Query")));
+				return;
+			}
+
+			bRequestedDependencies = true;
+
+			UE_CLOG(Dependencies.Num() > 0, LogGameFeatures, Verbose, TEXT("Found %i dependencies for %s"), Dependencies.Num(), *StateProperties.PluginName);
+
+			RemainingDependencies.Reserve(Dependencies.Num());
+			for (UGameFeaturePluginStateMachine* Dependency : Dependencies)
+			{
+				RemainingDependencies.Emplace(Dependency, MakeValue());
+				TransitionDependency(Dependency);
+			}
+		}
+
+		for (FDepResultPair& Pair : RemainingDependencies)
+		{
+			UGameFeaturePluginStateMachine* RemainingDependency = Pair.Key.Get();
+			if (!RemainingDependency)
+			{
+				// One of the dependency state machines was destroyed before finishing
+				StateStatus.SetTransitionError(TransitionPolicy::GetErrorState(), GetErrorResult(TEXT("Dependency_Destroyed_Before_Finish")));
+				return;
+			}
+
+			if (Pair.Value.HasError())
+			{
+				UE_LOG(LogGameFeatures, Error, TEXT("Dependency %s failed to transition with error %s"), *RemainingDependency->GetPluginURL(), *Pair.Value.GetError());
+				StateStatus.SetTransitionError(TransitionPolicy::GetErrorState(), GetErrorResult(TEXT("Failed_Dependency_Transition")));
+				return;
+			}
+		}
+
+		if (RemainingDependencies.Num() == 0)
+		{
+			StateStatus.SetTransition(TransitionPolicy::GetTransitionState());
+		}
+	}
+
+	void TransitionDependency(UGameFeaturePluginStateMachine* Dependency)
+	{
+		const bool bSetDestination = Dependency->SetDestination(TransitionPolicy::GetDependencyStateRange(),
+			FGameFeatureStateTransitionComplete::CreateRaw(this, &FTransitionDependenciesGameFeaturePluginState::OnDependencyTransitionComplete));
+
+		if (!bSetDestination)
+		{
+			const bool bCancelPending = Dependency->TryCancel(
+				FGameFeatureStateTransitionCanceled::CreateRaw(this, &FTransitionDependenciesGameFeaturePluginState::OnDependencyTransitionCanceled));
+			if (!ensure(bCancelPending))
+			{
+				OnDependencyTransitionComplete(Dependency, GetErrorResult(TEXT("Failed_Dependency_Transition")));
+			}
+		}
+	}
+
+	void OnDependencyTransitionCanceled(UGameFeaturePluginStateMachine* Dependency)
+	{
+		// Special case for terminal state since it cannot be exited, we need to make a new machine
+		if (Dependency->GetCurrentState() == EGameFeaturePluginState::Terminal)
+		{
+			// Inherit dep protocol options if possible
+			FGameFeatureProtocolOptions DepProtocolOptions;
+			EGameFeaturePluginProtocol DepProtocol = Dependency->GetPluginIdentifier().GetPluginProtocol();
+			if (DepProtocol == EGameFeaturePluginProtocol::InstallBundle && StateProperties.ProtocolOptions.HasSubtype<FInstallBundlePluginProtocolOptions>())
+			{
+				DepProtocolOptions = StateProperties.RecycleProtocolOptions();
+			}
+
+			UGameFeaturePluginStateMachine* NewMachine = UGameFeaturesSubsystem::Get().FindOrCreateGameFeaturePluginStateMachine(Dependency->GetPluginURL(), DepProtocolOptions);
+			checkf(NewMachine != Dependency, TEXT("Game Feature Plugin %s should have already been removed from subsystem!"), *Dependency->GetPluginURL());
+
+			const int32 Index = RemainingDependencies.IndexOfByPredicate([Dependency](const FDepResultPair& Pair)
+				{
+					return Pair.Key == Dependency;
+				});
+
+			check(Index != INDEX_NONE);
+			FDepResultPair& FoundDep = RemainingDependencies[Index];
+			FoundDep.Key = NewMachine;
+
+			Dependency->RemovePendingTransitionCallback(this);
+			Dependency->RemovePendingCancelCallback(this);
+
+			Dependency = NewMachine;
+		}
+
+		// Now that the transition has been canceled, retry reaching the desired destination
+		const bool bSetDestination = Dependency->SetDestination(TransitionPolicy::GetDependencyStateRange(),
+			FGameFeatureStateTransitionComplete::CreateRaw(this, &FTransitionDependenciesGameFeaturePluginState::OnDependencyTransitionComplete));
+
+		if (!ensure(bSetDestination))
+		{
+			OnDependencyTransitionComplete(Dependency, GetErrorResult(TEXT("Failed_Dependency_Transition")));
+		}
+	}
+
+	void OnDependencyTransitionComplete(UGameFeaturePluginStateMachine* Dependency, const UE::GameFeatures::FResult& Result)
+	{
+		const int32 Index = RemainingDependencies.IndexOfByPredicate([Dependency](const FDepResultPair& Pair)
+			{
+				return Pair.Key == Dependency;
+			});
+
+		if (ensure(Index != INDEX_NONE))
+		{
+			FDepResultPair& FoundDep = RemainingDependencies[Index];
+
+			if (Result.HasError())
+			{
+				FoundDep.Value = Result;
+			}
+			else
+			{
+				RemainingDependencies.RemoveAtSwap(Index, 1, false);
+			}
+
+			UpdateStateMachineImmediate();
+		}
+	}
+
+	void ClearDependencies()
+	{
+		for (FDepResultPair& Pair : RemainingDependencies)
+		{
+			UGameFeaturePluginStateMachine* RemainingDependency = Pair.Key.Get();
+			if (RemainingDependency)
+			{
+				RemainingDependency->RemovePendingTransitionCallback(this);
+				RemainingDependency->RemovePendingCancelCallback(this);
+			}
+		}
+
+		RemainingDependencies.Empty();
+		bRequestedDependencies = false;
+	}
+
+	using FDepResultPair = TPair<TWeakObjectPtr<UGameFeaturePluginStateMachine>, UE::GameFeatures::FResult>;
+	TArray<FDepResultPair> RemainingDependencies;
+	bool bRequestedDependencies = false;
+};
+
 struct FGameFeaturePluginState_Uninitialized : public FGameFeaturePluginState
 {
 	FGameFeaturePluginState_Uninitialized(FGameFeaturePluginStateMachineProperties& InStateProperties) : FGameFeaturePluginState(InStateProperties) {}
@@ -1687,188 +1866,38 @@
 	}
 };
 
-struct FGameFeaturePluginState_WaitingForDependencies : public FGameFeaturePluginState
+struct FWaitingForDependenciesTransitionPolicy
 {
-	FGameFeaturePluginState_WaitingForDependencies(FGameFeaturePluginStateMachineProperties& InStateProperties)
-		: FGameFeaturePluginState(InStateProperties)
-		, bRequestedDependencies(false)
-	{}
+	static bool GetPluginDependencyStateMachines(const FGameFeaturePluginStateMachineProperties& InStateProperties, TArray<UGameFeaturePluginStateMachine*>& OutDependencyMachines)
+	{
+		UGameFeaturesSubsystem& GameFeaturesSubsystem = UGameFeaturesSubsystem::Get();
 
-	virtual ~FGameFeaturePluginState_WaitingForDependencies()
-	{
-		ClearDependencies();
+		return GameFeaturesSubsystem.FindOrCreatePluginDependencyStateMachines(
+			*InStateProperties.PluginIdentifier.GetFullPluginURL(), InStateProperties.PluginInstalledFilename, InStateProperties.RecycleProtocolOptions(), OutDependencyMachines);
 	}
 
-	virtual void BeginState() override
+	static FGameFeaturePluginStateRange GetDependencyStateRange()
 	{
-		ClearDependencies();
-	}
-
-	virtual void EndState() override
-	{
-		ClearDependencies();
-	}
-
-	virtual void UpdateState(FGameFeaturePluginStateStatus& StateStatus) override
-	{
-		TRACE_CPUPROFILER_EVENT_SCOPE(GFP_WaitingForDependencies);
-		checkf(!StateProperties.PluginInstalledFilename.IsEmpty(), TEXT("PluginInstalledFilename must be set by the loading dependencies phase. PluginURL: %s"), *StateProperties.PluginIdentifier.GetFullPluginURL());
-		checkf(FPaths::GetExtension(StateProperties.PluginInstalledFilename) == TEXT("uplugin"), TEXT("PluginInstalledFilename must have a uplugin extension. PluginURL: %s"), *StateProperties.PluginIdentifier.GetFullPluginURL());
-
-		if (!bRequestedDependencies)
-		{
-			UGameFeaturesSubsystem& GameFeaturesSubsystem = UGameFeaturesSubsystem::Get();
-
-			TArray<UGameFeaturePluginStateMachine*> Dependencies;
-			if (!GameFeaturesSubsystem.FindOrCreatePluginDependencyStateMachines(
-				*StateProperties.PluginIdentifier.GetFullPluginURL(), StateProperties.PluginInstalledFilename, StateProperties.RecycleProtocolOptions(), Dependencies))
-			{
-				// Failed to query dependencies
-				StateStatus.SetTransitionError(EGameFeaturePluginState::ErrorWaitingForDependencies, GetErrorResult(TEXT("Failed_Dependency_Query")));
-				return;
-			}
-			
-			bRequestedDependencies = true;
-
-			UE_CLOG(Dependencies.Num() > 0, LogGameFeatures, Verbose, TEXT("Found %i dependencies for %s"), Dependencies.Num(), *StateProperties.PluginName);
-
-			RemainingDependencies.Reserve(Dependencies.Num());
-			for (UGameFeaturePluginStateMachine* Dependency : Dependencies)
-			{	
-				RemainingDependencies.Emplace(Dependency, MakeValue());
-				TransitionDependency(Dependency);
-			}
-		}
-
-		for (FDepResultPair& Pair : RemainingDependencies)
-		{
-			UGameFeaturePluginStateMachine* RemainingDependency = Pair.Key.Get();
-			if (!RemainingDependency)
-			{
-				// One of the dependency state machines was destroyed before finishing
-				StateStatus.SetTransitionError(EGameFeaturePluginState::ErrorWaitingForDependencies, GetErrorResult(TEXT("Dependency_Destroyed_Before_Finish")));
-				return;
-			}
-
-			if (Pair.Value.HasError())
-			{
-				UE_LOG(LogGameFeatures, Error, TEXT("Dependency %s failed to load with error %s"), *RemainingDependency->GetPluginURL(), *Pair.Value.GetError());
-				StateStatus.SetTransitionError(EGameFeaturePluginState::ErrorWaitingForDependencies, GetErrorResult(TEXT("Failed_Dependency_Register")));
-				return;
-			}
-		}
-
-		if (RemainingDependencies.Num() == 0)
-		{
-			StateStatus.SetTransition(EGameFeaturePluginState::Registering);
-		}
-	}
-
-	FGameFeaturePluginStateRange GetDependencyStateRange() const
-	{
 		return FGameFeaturePluginStateRange(EGameFeaturePluginState::Registered, EGameFeaturePluginState::Active);
 	}
 
-	void TransitionDependency(UGameFeaturePluginStateMachine* Dependency)
+	static EGameFeaturePluginState GetTransitionState()
 	{
-		const bool bSetDestination = Dependency->SetDestination(GetDependencyStateRange(),
-			FGameFeatureStateTransitionComplete::CreateRaw(this, &FGameFeaturePluginState_WaitingForDependencies::OnDependencyTransitionComplete));
-
-		if (!bSetDestination)
-		{
-			const bool bCancelPending = Dependency->TryCancel(
-				FGameFeatureStateTransitionCanceled::CreateRaw(this, &FGameFeaturePluginState_WaitingForDependencies::OnDependencyTransitionCanceled));
-			if (!ensure(bCancelPending))
-			{
-				OnDependencyTransitionComplete(Dependency, GetErrorResult(UE::GameFeatures::CommonErrorCodes::DependencyFailedRegister));
-			}
-		}
+		return EGameFeaturePluginState::Registering;
 	}
 
-	void OnDependencyTransitionCanceled(UGameFeaturePluginStateMachine* Dependency)
+	static EGameFeaturePluginState GetErrorState()
 	{
-		// Special case for terminal state since it cannot be exited, we need to make a new machine
-		if (Dependency->GetCurrentState() == EGameFeaturePluginState::Terminal)
-		{
-			// Inherit dep protocol options if possible
-			FGameFeatureProtocolOptions DepProtocolOptions;
-			EGameFeaturePluginProtocol DepProtocol = Dependency->GetPluginIdentifier().GetPluginProtocol();
-			if (DepProtocol == EGameFeaturePluginProtocol::InstallBundle && StateProperties.ProtocolOptions.HasSubtype<FInstallBundlePluginProtocolOptions>())
-			{
-				DepProtocolOptions = StateProperties.RecycleProtocolOptions();
-			}
-
-			UGameFeaturePluginStateMachine* NewMachine = UGameFeaturesSubsystem::Get().FindOrCreateGameFeaturePluginStateMachine(Dependency->GetPluginURL(), DepProtocolOptions);
-			checkf(NewMachine != Dependency, TEXT("Game Feature Plugin %s should have already been removed from subsystem!"), *Dependency->GetPluginURL());
-
-			const int32 Index = RemainingDependencies.IndexOfByPredicate([Dependency](const FDepResultPair& Pair)
-				{
-					return Pair.Key == Dependency;
-				});
-
-			check(Index != INDEX_NONE);
-			FDepResultPair& FoundDep = RemainingDependencies[Index];
-			FoundDep.Key = NewMachine;
-
-			Dependency->RemovePendingTransitionCallback(this);
-			Dependency->RemovePendingCancelCallback(this);
-
-			Dependency = NewMachine;
-		}
-
-		// Now that the transition has been canceled, retry reaching the desired destination
-		const bool bSetDestination = Dependency->SetDestination(GetDependencyStateRange(),
-			FGameFeatureStateTransitionComplete::CreateRaw(this, &FGameFeaturePluginState_WaitingForDependencies::OnDependencyTransitionComplete));
-
-		if (!ensure(bSetDestination))
-		{
-			OnDependencyTransitionComplete(Dependency, GetErrorResult(UE::GameFeatures::CommonErrorCodes::DependencyFailedRegister));
-		}
+		return EGameFeaturePluginState::ErrorWaitingForDependencies;
 	}
-
-	void OnDependencyTransitionComplete(UGameFeaturePluginStateMachine* Dependency, const UE::GameFeatures::FResult& Result)
-	{
-		const int32 Index = RemainingDependencies.IndexOfByPredicate([Dependency](const FDepResultPair& Pair)
-		{
-			return Pair.Key == Dependency;
-		});
-
-		if (ensure(Index != INDEX_NONE))
-		{
-			FDepResultPair& FoundDep = RemainingDependencies[Index];
-			
-			if (Result.HasError())
-			{
-				FoundDep.Value = Result;
-			}
-			else
-			{
-				RemainingDependencies.RemoveAtSwap(Index, 1, false);
-			}
-
-			UpdateStateMachineImmediate();
-		}
-	}
-
-	void ClearDependencies()
-	{
-		for (FDepResultPair& Pair : RemainingDependencies)
-		{
-			UGameFeaturePluginStateMachine* RemainingDependency = Pair.Key.Get();
-			if (RemainingDependency)
-			{
-				RemainingDependency->RemovePendingTransitionCallback(this);
-				RemainingDependency->RemovePendingCancelCallback(this);
-			}
-		}
-
-		RemainingDependencies.Empty();
-		bRequestedDependencies = false;
-	}
-
-	using FDepResultPair = TPair<TWeakObjectPtr<UGameFeaturePluginStateMachine>, UE::GameFeatures::FResult>;
-	TArray<FDepResultPair> RemainingDependencies;
-	bool bRequestedDependencies = false;
 };
 
+struct FGameFeaturePluginState_WaitingForDependencies : public FTransitionDependenciesGameFeaturePluginState<FWaitingForDependenciesTransitionPolicy>
+{
+	FGameFeaturePluginState_WaitingForDependencies(FGameFeaturePluginStateMachineProperties& InStateProperties)
+		: FTransitionDependenciesGameFeaturePluginState(InStateProperties)
+	{
+	}
+};
+
 struct FGameFeaturePluginState_Unregistering : public FGameFeaturePluginState
@@ -2165,7 +2194,7 @@
 	{
 		if (StateProperties.Destination > EGameFeaturePluginState::Loaded)
 		{
-			StateStatus.SetTransition(EGameFeaturePluginState::Activating);
+			StateStatus.SetTransition(EGameFeaturePluginState::ActivatingDependencies);
 		}
 		else if (StateProperties.Destination < EGameFeaturePluginState::Loaded)
 		{
@@ -2262,6 +2291,58 @@
 	}
 };
 
+struct FGameFeaturePluginState_ErrorActivatingDependencies : public FErrorGameFeaturePluginState
+{
+	FGameFeaturePluginState_ErrorActivatingDependencies(FGameFeaturePluginStateMachineProperties& InStateProperties) : FErrorGameFeaturePluginState(InStateProperties) {}
+
+	virtual void UpdateState(FGameFeaturePluginStateStatus& StateStatus) override
+	{
+		if (StateProperties.Destination < EGameFeaturePluginState::ErrorActivatingDependencies)
+		{
+			// There is no cleaup state equivalent to EGameFeaturePluginState::ErrorActivatingDependencies so just go back to Unloading
+			StateStatus.SetTransition(EGameFeaturePluginState::Unloading);
+		}
+		else if (StateProperties.Destination > EGameFeaturePluginState::ErrorActivatingDependencies)
+		{
+			StateStatus.SetTransition(EGameFeaturePluginState::ActivatingDependencies);
+		}
+	}
+};
+
+struct FActivatingDependenciesTransitionPolicy
+{
+	static bool GetPluginDependencyStateMachines(const FGameFeaturePluginStateMachineProperties& InStateProperties, TArray<UGameFeaturePluginStateMachine*>& OutDependencyMachines)
+	{
+		UGameFeaturesSubsystem& GameFeaturesSubsystem = UGameFeaturesSubsystem::Get();
+
+		return GameFeaturesSubsystem.FindPluginDependencyStateMachinesToActivate(
+			*InStateProperties.PluginIdentifier.GetFullPluginURL(), InStateProperties.PluginInstalledFilename, OutDependencyMachines);
+	}
+
+	static FGameFeaturePluginStateRange GetDependencyStateRange()
+	{
+		return FGameFeaturePluginStateRange(EGameFeaturePluginState::Active, EGameFeaturePluginState::Active);
+	}
+
+	static EGameFeaturePluginState GetTransitionState()
+	{
+		return EGameFeaturePluginState::Activating;
+	}
+
+	static EGameFeaturePluginState GetErrorState()
+	{
+		return EGameFeaturePluginState::ErrorActivatingDependencies;
+	}
+};
+
+struct FGameFeaturePluginState_ActivatingDependencies : public FTransitionDependenciesGameFeaturePluginState<FActivatingDependenciesTransitionPolicy>
+{
+	FGameFeaturePluginState_ActivatingDependencies(FGameFeaturePluginStateMachineProperties& InStateProperties)
+		: FTransitionDependenciesGameFeaturePluginState(InStateProperties)
+	{
+	}
+};
+
 struct FGameFeaturePluginState_Activating : public FGameFeaturePluginState
 {
 	FGameFeaturePluginState_Activating(FGameFeaturePluginStateMachineProperties& InStateProperties) : FGameFeaturePluginState(InStateProperties) {}
diff --git a/Engine/Plugins/Experimental/GameFeatures/Source/GameFeatures/Private/GameFeaturePluginStateMachine.h b/Engine/Plugins/Experimental/GameFeatures/Source/GameFeatures/Private/GameFeaturePluginStateMachine.h
--- a/Engine/Plugins/Experimental/GameFeatures/Source/GameFeatures/Private/GameFeaturePluginStateMachine.h
+++ b/Engine/Plugins/Experimental/GameFeatures/Source/GameFeatures/Private/GameFeaturePluginStateMachine.h
@@ -120,24 +120,32 @@ Transition states are expected to transition the machine to another state after
                |            |
                +-^--------+-+
                  |        |
-        +--------+--+  +--v--------+
-        |           |  |           |
-        | Unloading |  |  Loading  |
-        |           |  |           |
-        +--------^--+  +--+--------+
-                 |        |
-               +-+--------v-+
-               |      *     |
-               |   Loaded   |
-               |            |
-               +-^--------+-+
-                 |        |
-        +--------+---+  +-v---------+
-        |            |  |           |
-        |Deactivating|  |Activating |
-        |            |  |           |
-        +--------^---+  +-+---------+
+           ------~--------~------------------------------
+           |     |        |                             |
+        +--v-----+--+  +--v--------+                    |
+        |           |  |           |                    |
+        | Unloading |  |  Loading  |                    |
+        |           |  |           |                    |
+        +--------^--+  +--+--------+                    |
+                 |        |                             |
+               +-+--------v-+                           |
+               |      *     |                           |
+               |   Loaded   |                           |
+               |            |                           |
+               +-^--------+-+                           |
+                 |        |                             |
+        +--------+---+  +-v----------------------+   +--+--------------------------+
+        |            |  |                        |   |             !               |
+        |Deactivating|  | ActivatingDependencies <---> ErrorActivatingDependencies |
+        |            |  |                        |   |                             |
+        +--------^---+  +-+----------------------+   +-----------------------------+
                  |        |
+				 |      +-v----------+
+				 |      |            |
+				 |      | Activating |
+				 |      |            |
+				 |      +-+----------+
+				 |        |
                +-+--------v-+
                |      *     |
                |   Active   |
diff --git a/Engine/Plugins/Experimental/GameFeatures/Source/GameFeatures/Private/GameFeaturesSubsystem.cpp b/Engine/Plugins/Experimental/GameFeatures/Source/GameFeatures/Private/GameFeaturesSubsystem.cpp
--- a/Engine/Plugins/Experimental/GameFeatures/Source/GameFeatures/Private/GameFeaturesSubsystem.cpp
+++ b/Engine/Plugins/Experimental/GameFeatures/Source/GameFeatures/Private/GameFeaturesSubsystem.cpp
@@ -1747,6 +1747,7 @@ bool UGameFeaturesSubsystem::GetGameFeaturePluginDetails(const FString& PluginUR
 	{
 		FString NameField = TEXT("Name");
 		FString EnabledField = TEXT("Enabled");
+		FString ActivateField = TEXT("Activate");
 		for (const TSharedPtr<FJsonValue>& PluginElement : *PluginsArray)
 		{
 			if (PluginElement.IsValid())
@@ -1759,6 +1760,9 @@ bool UGameFeaturesSubsystem::GetGameFeaturePluginDetails(const FString& PluginUR
 					bool bElementEnabled = false;
 					ElementObject->TryGetBoolField(EnabledField, bElementEnabled);
 
+					bool bElementActivate = false;
+					ElementObject->TryGetBoolField(ActivateField, bElementActivate);
+
 					if (bElementEnabled)
 					{
 						FString DependencyName;
@@ -1772,7 +1776,7 @@ bool UGameFeaturesSubsystem::GetGameFeaturePluginDetails(const FString& PluginUR
 							}
 							else if (ResolvedDepResult.HasValue() && !ResolvedDepResult.GetValue().IsEmpty()) // Dependency may not be a GFP
 							{
-								OutPluginDetails.PluginDependencies.Add(ResolvedDepResult.StealValue());
+								OutPluginDetails.PluginDependencies.Emplace(FGameFeaturePluginReferenceDetails(ResolvedDepResult.StealValue(), bElementActivate));
 							}
 						}
 					}
@@ -1981,17 +1985,17 @@ bool UGameFeaturesSubsystem::FindOrCreatePluginDependencyStateMachines(const FSt
 	FGameFeaturePluginDetails Details;
 	if (GetGameFeaturePluginDetails(PluginURL, PluginFilename, Details))
 	{
-		for (const FString& DependencyURL : Details.PluginDependencies)
+		for (const FGameFeaturePluginReferenceDetails& PluginDependency : Details.PluginDependencies)
 		{
 			// Inherit dep protocol options if possible
 			FGameFeatureProtocolOptions DepProtocolOptions;
-			EGameFeaturePluginProtocol DepProtocol = UGameFeaturesSubsystem::GetPluginURLProtocol(DependencyURL);
+			EGameFeaturePluginProtocol DepProtocol = UGameFeaturesSubsystem::GetPluginURLProtocol(PluginDependency.URL);
 			if (DepProtocol == EGameFeaturePluginProtocol::InstallBundle && InDepProtocolOptions.HasSubtype<FInstallBundlePluginProtocolOptions>())
 			{
 				DepProtocolOptions = InDepProtocolOptions;
 			}
 
-			UGameFeaturePluginStateMachine* Dependency = FindOrCreateGameFeaturePluginStateMachine(DependencyURL, DepProtocolOptions);
+			UGameFeaturePluginStateMachine* Dependency = FindOrCreateGameFeaturePluginStateMachine(PluginDependency.URL, DepProtocolOptions);
 			check(Dependency);
 			OutDependencyMachines.Add(Dependency);
 		}
@@ -2002,6 +2006,27 @@ bool UGameFeaturesSubsystem::FindOrCreatePluginDependencyStateMachines(const FSt
 	return false;
 }
 
+bool UGameFeaturesSubsystem::FindPluginDependencyStateMachinesToActivate(const FString& PluginURL, const FString& PluginFilename, TArray<UGameFeaturePluginStateMachine*>& OutDependencyMachines)
+{
+	FGameFeaturePluginDetails Details;
+	if (GetGameFeaturePluginDetails(PluginURL, PluginFilename, Details))
+	{
+		for (const FGameFeaturePluginReferenceDetails& PluginDependency : Details.PluginDependencies)
+		{
+			if (PluginDependency.bShouldActivate)
+			{
+				UGameFeaturePluginStateMachine* Dependency = FindGameFeaturePluginStateMachine(PluginDependency.URL);
+				check(Dependency);
+				OutDependencyMachines.Add(Dependency);
+			}
+		}
+
+		return true;
+	}
+
+	return false;
+}
+
 void UGameFeaturesSubsystem::ListGameFeaturePlugins(const TArray<FString>& Args, UWorld* InWorld, FOutputDevice& Ar)
 {
 	const bool bAlphaSort = Args.ContainsByPredicate([](const FString& Arg) { return Arg.Compare(TEXT("-ALPHASORT"), ESearchCase::IgnoreCase) == 0; });
diff --git a/Engine/Plugins/Experimental/GameFeatures/Source/GameFeatures/Public/GameFeatureTypes.h b/Engine/Plugins/Experimental/GameFeatures/Source/GameFeatures/Public/GameFeatureTypes.h
--- a/Engine/Plugins/Experimental/GameFeatures/Source/GameFeatures/Public/GameFeatureTypes.h
+++ b/Engine/Plugins/Experimental/GameFeatures/Source/GameFeatures/Public/GameFeatureTypes.h
@@ -33,6 +33,8 @@ class FString;
 	XSTATE(Unloading,					NSLOCTEXT("GameFeatures", "UnloadingStateDisplayName", "Unloading"))										/* Transition state Loaded -> Registered. In the process of removing code/content from memory. */ \
 	XSTATE(Loading,						NSLOCTEXT("GameFeatures", "LoadingStateDisplayName", "Loading"))											/* Transition state Registered -> Loaded. In the process of loading code/content into memory. */ \
 	XSTATE(Loaded,						NSLOCTEXT("GameFeatures", "LoadedStateDisplayName", "Loaded"))												/* The plugin is loaded into memory, but not registered with game systems and active. */ \
+	XSTATE(ErrorActivatingDependencies,	NSLOCTEXT("GameFeatures", "ErrorActivatingDependenciesStateDisplayName", "ErrorActivatingDependencies"))	/* Error state for Registered -> Active transition. */ \
+	XSTATE(ActivatingDependencies,		NSLOCTEXT("GameFeatures", "ActivatingDependenciesStateDisplayName", "ActivatingDependencies"))				/* Transition state Registered -> Active. In the process of selectively activating dependencies.*/ \
 	XSTATE(Deactivating,				NSLOCTEXT("GameFeatures", "DeactivatingStateDisplayName", "Deactivating"))									/* Transition state Active -> Loaded. Currently unregistering with game systems. */ \
 	XSTATE(Activating,					NSLOCTEXT("GameFeatures", "ActivatingStateDisplayName", "Activating"))										/* Transition state Loaded -> Active. Currently registering plugin code/content with game systems. */ \
 	XSTATE(Active,						NSLOCTEXT("GameFeatures", "ActiveStateDisplayName", "Active"))												/* Plugin is fully loaded and active. It is affecting the game.  */ 
diff --git a/Engine/Plugins/Experimental/GameFeatures/Source/GameFeatures/Public/GameFeaturesSubsystem.h b/Engine/Plugins/Experimental/GameFeatures/Source/GameFeatures/Public/GameFeaturesSubsystem.h
--- a/Engine/Plugins/Experimental/GameFeatures/Source/GameFeatures/Public/GameFeaturesSubsystem.h
+++ b/Engine/Plugins/Experimental/GameFeatures/Source/GameFeatures/Public/GameFeaturesSubsystem.h
@@ -193,9 +193,21 @@ enum class EGameFeatureTargetState : uint8
 const FString GAMEFEATURES_API LexToString(const EGameFeatureTargetState GameFeatureTargetState);
 void GAMEFEATURES_API LexFromString(EGameFeatureTargetState& Value, const TCHAR* StringIn);
 
+struct FGameFeaturePluginReferenceDetails
+{
+	FString URL;
+	bool bShouldActivate;
+
+	FGameFeaturePluginReferenceDetails(const FString& InURL, bool bInShouldActivate)
+		: URL(InURL)
+		, bShouldActivate(bInShouldActivate)
+	{
+	}
+};
+
 struct FGameFeaturePluginDetails
 {
-	TArray<FString> PluginDependencies;
+	TArray<FGameFeaturePluginReferenceDetails> PluginDependencies;
 	TMap<FString, FString> AdditionalMetadata;
 	bool bHotfixable;
 	EBuiltInAutoState BuiltInAutoState;
@@ -596,7 +608,11 @@ private:
 
 	/** Handler for when a state machine requests its dependencies. Returns false if the dependencies could not be read */
 	bool FindOrCreatePluginDependencyStateMachines(const FString& PluginURL, const FString& PluginFilename, const FGameFeatureProtocolOptions& DepProtocolOptions, TArray<UGameFeaturePluginStateMachine*>& OutDependencyMachines);
-	friend struct FGameFeaturePluginState_WaitingForDependencies;
+	template <typename> friend struct FTransitionDependenciesGameFeaturePluginState;
+	friend struct FWaitingForDependenciesTransitionPolicy;
+
+	bool FindPluginDependencyStateMachinesToActivate(const FString& PluginURL, const FString& PluginFilename, TArray<UGameFeaturePluginStateMachine*>& OutDependencyMachines);
+	friend struct FActivatingDependenciesTransitionPolicy;
 
 	/** Handle 'ListGameFeaturePlugins' console command */
 	void ListGameFeaturePlugins(const TArray<FString>& Args, UWorld* InWorld, FOutputDevice& Ar);
diff --git a/Engine/Plugins/Experimental/GameFeatures/Source/GameFeaturesEditor/Private/SGameFeatureStateWidget.cpp b/Engine/Plugins/Experimental/GameFeatures/Source/GameFeaturesEditor/Private/SGameFeatureStateWidget.cpp
--- a/Engine/Plugins/Experimental/GameFeatures/Source/GameFeaturesEditor/Private/SGameFeatureStateWidget.cpp
+++ b/Engine/Plugins/Experimental/GameFeatures/Source/GameFeaturesEditor/Private/SGameFeatureStateWidget.cpp
@@ -65,7 +65,7 @@ FText SGameFeatureStateWidget::GetDisplayNameOfState(EGameFeaturePluginState Sta
 
 FText SGameFeatureStateWidget::GetTooltipOfState(EGameFeaturePluginState State)
 {
-	static_assert((int32)EGameFeaturePluginState::MAX == 29, "");
+	static_assert((int32)EGameFeaturePluginState::MAX == 31, "");
 
 	switch (State)
 	{
@@ -121,6 +121,10 @@ FText SGameFeatureStateWidget::GetTooltipOfState(EGameFeaturePluginState State)
 		return LOCTEXT("StateTooltip_Loading", "Transition state Registered -> Loaded. In the process of loading code/content into memory.");
 	case EGameFeaturePluginState::Loaded:
 		return LOCTEXT("StateTooltip_Loaded", "The plugin is loaded into memory and registered with some game systems but not yet active.");
+	case EGameFeaturePluginState::ErrorActivatingDependencies:
+		return LOCTEXT("StateTooltip_ErrorActivateDependencies", "Error state for Loaded -> Active and Active -> Loaded transitions.");
+	case EGameFeaturePluginState::ActivatingDependencies:
+		return LOCTEXT("StateTooltip_ActivateDependencies", "Transition state Loaded -> Active. In the process of selectively activating dependencies.");
 	case EGameFeaturePluginState::Deactivating:
 		return LOCTEXT("StateTooltip_Deactivating", "Transition state Active -> Loaded. Currently unregistering with game systems.");
 	case EGameFeaturePluginState::Activating:
