Add support for custom device profiles in game feature plugins

diff --git a/Engine/Plugins/Experimental/GameFeatures/Source/GameFeatures/Private/GameFeatureData.cpp b/Engine/Plugins/Experimental/GameFeatures/Source/GameFeatures/Private/GameFeatureData.cpp
--- a/Engine/Plugins/Experimental/GameFeatures/Source/GameFeatures/Private/GameFeatureData.cpp
+++ b/Engine/Plugins/Experimental/GameFeatures/Source/GameFeatures/Private/GameFeatureData.cpp
@@ -8,6 +8,9 @@
 #include "Misc/ConfigCacheIni.h"
 #include "Misc/ConfigContext.h"
 #include "UObject/CoreRedirects.h"
+#include "DeviceProfiles/DeviceProfile.h"
+#include "DeviceProfiles/DeviceProfileFragment.h"
+#include "DeviceProfiles/DeviceProfileManager.h"
 #include "GameFeatureAction_AddWPContent.h"
 #include "WorldPartition/ContentBundle/ContentBundleDescriptor.h"
 #include "AssetRegistry/AssetData.h"
@@ -66,6 +69,12 @@ EDataValidationResult UGameFeatureData::IsDataValid(FDataValidationContext& Cont
 }
 #endif
 
+static TAutoConsoleVariable<bool> CVarAllowRuntimeDeviceProfiles(
+	TEXT("GameFeaturePlugin.AllowRuntimeDeviceProfiles"),
+	true,
+	TEXT("Allow game feature plugins to generate device profiles from config based on existing parents"),
+	ECVF_Default);
+
 void UGameFeatureData::InitializeBasePluginIniFile(const FString& PluginInstalledFilename) const
 {
 	const FString PluginName = FPaths::GetBaseFilename(PluginInstalledFilename);
@@ -110,47 +119,283 @@ void UGameFeatureData::InitializeBasePluginIniFile(const FString& PluginInstalle
 
 void UGameFeatureData::InitializeHierarchicalPluginIniFiles(const FString& PluginInstalledFilename) const
 {
+	UDeviceProfileManager& DeviceProfileManager = UDeviceProfileManager::Get();
+
+#if ALLOW_OTHER_PLATFORM_CONFIG
+	const UDeviceProfile* PreviewDeviceProfile = DeviceProfileManager.GetPreviewDeviceProfile();
+	const FString PlatformName = PreviewDeviceProfile ? PreviewDeviceProfile->ConfigPlatform : FPlatformProperties::IniPlatformName();
+#else
+	const FString PlatformName = FPlatformProperties::IniPlatformName();
+#endif
+
 	const FString PluginName = FPaths::GetBaseFilename(PluginInstalledFilename);
 	const FString PluginConfigDir = FPaths::GetPath(PluginInstalledFilename) / TEXT("Config/");
+	const FString PlatformExtensionDir = FPaths::ProjectPlatformExtensionsDir() / (PlatformName + "/");
 	const FString EngineConfigDir = FPaths::EngineConfigDir();
 
+	// We'll look first in the game's platform extension dir for a plugin filesystem, and then default to the plugin folder
+	FString PluginPlatformConfigDir = FPaths::GetPath(PluginInstalledFilename).Replace(*FPaths::ProjectDir(), *PlatformExtensionDir) / TEXT("Config");
+	if (!FPaths::DirectoryExists(PluginPlatformConfigDir))
+	{
+		PluginPlatformConfigDir = FPaths::Combine(PluginConfigDir, PlatformName);
+	}
+
 	const bool bIsBaseIniName = false;
 	const bool bForceReloadFromDisk = false;
 	const bool bWriteDestIni = false;
+	const bool bCreateDeviceProfiles = CVarAllowRuntimeDeviceProfiles->GetBool();
 
-	// @todo: Likely we need to track the diffs this config caused and/or store versions/layers in order to unwind settings during unloading/deactivation
-	TArray<FString> IniNamesToLoad = { 
-		TEXT("Input"), TEXT("Game"), TEXT("Engine")
+	struct FIniLoadingParams
+	{
+		FIniLoadingParams(const FString& InName, bool bInUsePlatformDir = false, bool bInCreateDeviceProfiles = false)
+			: Name(InName), bUsePlatformDir(bInUsePlatformDir), bCreateDeviceProfiles(bInCreateDeviceProfiles)
+		{}
+
+		FString Name;
+		bool bUsePlatformDir;
+		bool bCreateDeviceProfiles;
+	};
+
+	// Engine.ini and DeviceProfiles.ini will also support platform extensions
+	TArray<FIniLoadingParams> IniFilesToLoad = { FIniLoadingParams(TEXT("Input")),
+		FIniLoadingParams(TEXT("Game")), FIniLoadingParams(TEXT("Game"), true),
+		FIniLoadingParams(TEXT("Engine")), FIniLoadingParams(TEXT("Engine"), true),
 #if UE_EDITOR
-		, TEXT("Editor")
+		FIniLoadingParams(TEXT("Editor")),
 #endif
+		FIniLoadingParams(TEXT("DeviceProfiles"), false, bCreateDeviceProfiles),
+		FIniLoadingParams(TEXT("DeviceProfiles"), true, bCreateDeviceProfiles)
+	};
+
+	// Create overridden device profiles for each matching rule in config
+	auto InsertRuntimeDeviceProfilesIntoConfig = [&](FConfigFile& PluginConfig, const FConfigFile& ExistingConfig)
+	{
+		TMap<FString, FConfigSection> ConfigSectionsToAdd;
+
+		for (auto& Section : PluginConfig)
+		{
+			FString RuleName, ParentClass;
+			if (Section.Key.Split(TEXT(" "), &RuleName, &ParentClass))
+			{
+				// Early reject anything that's not handled in here
+				const bool bIsRuntimeDeviceProfileRule = (ParentClass == "RuntimeDeviceProfileRule");
+				const bool bIsDeviceProfileFragment = (ParentClass == UDeviceProfileFragment::StaticClass()->GetName());
+				if (!(bIsRuntimeDeviceProfileRule || bIsDeviceProfileFragment))
+				{
+					continue;
+				}
+
+				// Check the existing config because at this point in time it has already been hotfixed from the base empty config
+				// Those CVars are also always without a +-. prefix because they're the result of the hotfix applied to the empty config.
+				// @todo: we cannot use hotfixes to remove CVars because HF happens way before GFPs have a chance to load.
+				// The hotfix process is destructive and doesn't leave us an opportunity to read the hotfix delta when loading GFPs.
+				TArray<FConfigValue> HotfixCVars;
+				if (const FConfigSection* HotfixSection = ExistingConfig.Find(Section.Key))
+				{
+					HotfixSection->MultiFind("CVars", HotfixCVars);
+				}
+
+				// Extract key-value pairs for CVars and FragmentIncludes, keeping the +-. prefix
+				TMultiMap<FName, FConfigValue> PluginCVars;
+				TMultiMap<FName, FConfigValue> FragmentIncludes;
+				for (const auto& Entry : Section.Value)
+				{
+					const FString& EntryKey = Entry.Key.ToString();
+					if (EntryKey.RightChop(1).StartsWith("CVars"))
+					{
+						PluginCVars.Add(Entry.Key, Entry.Value);
+					}
+					else if (EntryKey.RightChop(1).StartsWith("FragmentIncludes"))
+					{
+						FragmentIncludes.Add(Entry.Key, Entry.Value);
+					}
+				}
+
+				// Check if a CVar should be either included, or removed for a new hotfix value
+				auto ShouldKeepCVar = [&HotfixCVars](const FString& Key, FString& Value) -> bool
+				{
+					for (const FConfigValue& HotfixCVarData : HotfixCVars)
+					{
+						FString HotfixCVarKey, HotfixCVarValue;
+						if (HotfixCVarData.GetValue().Split(TEXT("="), &HotfixCVarKey, &HotfixCVarValue) && HotfixCVarKey == Key && HotfixCVarValue != Value)
+						{
+							Value = HotfixCVarValue;
+							return false;
+						}
+					}
+					return true;
+				};
+
+				// Process new runtime device profile
+				if (bIsRuntimeDeviceProfileRule)
+				{
+					UE_LOG(LogGameFeatures, Display, TEXT("Game feature '%s' found runtime device profile rule %s"), *PluginName, *RuleName);
+
+					// Extract metadata
+					const FConfigValue* ParentProfileName = Section.Value.Find("ParentProfileName");
+					const FConfigValue* ProfileSuffix = Section.Value.Find("ProfileSuffix");
+
+					if (ParentProfileName && ProfileSuffix && (FragmentIncludes.Num() > 0 || PluginCVars.Num() > 0))
+					{
+						for (const UDeviceProfile* Profile : DeviceProfileManager.Profiles)
+						{
+							// Check if one of the parents for this profile is the one this rule applies to
+							bool bProfileHasCompatibleParent = false;
+							const UDeviceProfile* CurrentProfile = Profile;
+							do
+							{
+								bProfileHasCompatibleParent = CurrentProfile->GetName() == ParentProfileName->GetValue();
+								CurrentProfile = CurrentProfile->GetParentProfile();
+							} while (!bProfileHasCompatibleParent && CurrentProfile != nullptr);
+
+							// Create the config for a runtime profile
+							if (bProfileHasCompatibleParent && !Profile->GetName().EndsWith(ProfileSuffix->GetValue()))
+							{
+								FConfigSection RuntimeProfile;
+								RuntimeProfile.Add("DeviceType", PlatformName);
+								RuntimeProfile.Add("BaseProfileName", FConfigValue(Profile->GetName()));
+
+								// Add fragment includes
+								for (const auto& FragmentInclude : FragmentIncludes)
+								{
+									RuntimeProfile.Add(FragmentInclude.Key, FConfigValue(FragmentInclude.Value.GetValue()));
+								}
+
+								// Add direct CVars
+								for (const auto& CVar : PluginCVars)
+								{
+									FString CVarKey, CVarValue;
+									if (CVar.Value.GetValue().Split(TEXT("="), &CVarKey, &CVarValue) && ShouldKeepCVar(CVarKey, CVarValue))
+									{
+										UE_LOG(LogGameFeatures, Verbose, TEXT(" Found CVar: %s=%s"), *CVarKey, *CVarValue);
+										RuntimeProfile.Add(CVar.Key, FConfigValue(CVarKey + "=" + CVarValue));
+									}
+								}
+
+								// Add hotfix CVars
+								for (const auto& CVar : HotfixCVars)
+								{
+									FString CVarKey, CVarValue;
+									if (CVar.GetValue().Split(TEXT("="), &CVarKey, &CVarValue) && !PluginCVars.Contains(FName(*CVarKey)))
+									{
+										UE_LOG(LogGameFeatures, Verbose, TEXT(" Added CVar: %s=%s"), *CVarKey, *CVarValue);
+										RuntimeProfile.Add(FName(*CVarKey), FConfigValue(CVarKey + "=" + CVarValue));
+									}
+								}
+
+								ConfigSectionsToAdd.Add(Profile->GetName() + ProfileSuffix->GetValue() + TEXT(" ") + UDeviceProfile::StaticClass()->GetName(), RuntimeProfile);
+							}
+						}
+					}
+					else
+					{
+						UE_LOG(LogGameFeatures, Warning, TEXT("Game feature '%s' has invalid runtime device profile %s with parent %s, suffix %s, %d CVars, %d fragments"),
+							*PluginName, ParentProfileName ? *ParentProfileName->GetValue() : TEXT("null"), ProfileSuffix ? *ProfileSuffix->GetValue() : TEXT("null"),
+							FragmentIncludes.Num(), PluginCVars.Num());
+					}
+				}
+
+				// Hotfix device profile fragments
+				else if (bIsDeviceProfileFragment && HotfixCVars.Num() > 0)
+				{
+					UE_LOG(LogGameFeatures, Display, TEXT("Game feature '%s' found device profile fragment %s"), *PluginName, *RuleName);
+
+					// Update existing CVars
+					for (const auto& CVar : PluginCVars)
+					{
+						FString CVarKey, CVarValue;
+						if (CVar.Value.GetValue().Split(TEXT("="), &CVarKey, &CVarValue) && !ShouldKeepCVar(CVarKey, CVarValue))
+						{
+							UE_LOG(LogGameFeatures, Verbose, TEXT(" Removed CVar: %s=%s"), *CVarKey, *CVarValue);
+							Section.Value.Remove(CVar.Key);
+						}
+						else
+						{
+							UE_LOG(LogGameFeatures, Verbose, TEXT(" Found CVar: %s=%s"), *CVarKey, *CVarValue);
+						}
+					}
+
+					// Add new hotfix CVars
+					for (const auto& CVar : HotfixCVars)
+					{
+						FString CVarKey, CVarValue;
+						if (CVar.GetValue().Split(TEXT("="), &CVarKey, &CVarValue) && !PluginCVars.Contains(FName(*CVarKey)))
+						{
+							UE_LOG(LogGameFeatures, Verbose, TEXT(" Added CVar: %s=%s"), *CVarKey, *CVarValue);
+							Section.Value.Add(FName(*CVarKey), FConfigValue(CVarKey + "=" + CVarValue));
+						}
+					}
+				}
+			}
+		}
+
+		PluginConfig.Append(ConfigSectionsToAdd);
 	};
-	for (const FString& IniName : IniNamesToLoad)
+
+	// Create device profiles for this plugin from config
+	auto LoadDeviceProfilesFromConfig = [&](const FConfigFile& Config)
 	{
-		const FString PluginIniName = PluginName + IniName;
+		for (TPair<const FString&, const FConfigSection&> Section : Config)
+		{
+			FString ProfileName;
+			FString ParentClass;
+			if (Section.Key.Split(TEXT(" "), &ProfileName, &ParentClass) && ParentClass == UDeviceProfile::StaticClass()->GetName())
+			{
+				const FConfigValue* DeviceType = Section.Value.Find("DeviceType");
+				if (DeviceType)
+				{
+					UE_LOG(LogGameFeatures, Display, TEXT("Game feature '%s' adding new device profile %s"), *PluginName, *ProfileName);
+					DeviceProfileManager.CreateProfile(ProfileName, DeviceType->GetValue());
+				}
+			}
+		}
+	};
+
+	// @todo: Likely we need to track the diffs this config caused and/or store versions/layers in order to unwind settings during unloading/deactivation
+	for (const FIniLoadingParams& Ini : IniFilesToLoad)
+	{
+		const FString PluginIniName = Ini.bUsePlatformDir ? (PlatformName + PluginName + Ini.Name) : PluginName + Ini.Name;
+		const FString ConfigDirectory = Ini.bUsePlatformDir ? PluginPlatformConfigDir : PluginConfigDir;
+
 		// @note: Loading the INI in this manner in order to have a record of relevant sections that were changed so that affected objects can be reloaded. By virtue of how
 		// this is parsed (standalone instead of being treated as a combined diff), the actual data within the sections will likely be incorrect. As an example, users adding
 		// to an array with the "+" syntax will have the "+" incorrectly embedded inside the data in the temp FConfigFile. It's properly handled in the Combine() below where the
 		// actual INI changes are computed.
-		FConfigFile TempConfig;
-		if (FConfigCacheIni::LoadExternalIniFile(TempConfig, *PluginIniName, *EngineConfigDir, *PluginConfigDir, bIsBaseIniName, nullptr, bForceReloadFromDisk, bWriteDestIni) && (TempConfig.Num() > 0))
+		FConfigFile Config;
+		if (FConfigCacheIni::LoadExternalIniFile(Config, *PluginIniName, *EngineConfigDir, *ConfigDirectory, bIsBaseIniName, nullptr, bForceReloadFromDisk, bWriteDestIni) && (Config.Num() > 0))
 		{
+			UE_LOG(LogGameFeatures, Display, TEXT("Game feature '%s' loaded config file %s"), *PluginName, *PluginIniName);
+
 			// Need to get the in-memory config filename, the on disk one is likely not up to date
-			FString IniFile = GConfig->GetConfigFilename(*IniName);
+			FString IniFile = GConfig->GetConfigFilename(*Ini.Name);
 
 			if (FConfigFile* ExistingConfig = GConfig->FindConfigFile(IniFile))
 			{
+				if (Ini.bCreateDeviceProfiles)
+				{
+					InsertRuntimeDeviceProfilesIntoConfig(Config, *ExistingConfig);
+				}
+
+				FString ConfigAsString;
+				Config.WriteToString(ConfigAsString, PluginIniName);
+
 				// @todo: Might want to consider modifying the engine level's API here to allow for a combination that yields affected
 				// sections and/or optionally just does the reload itself. This route is less efficient than it needs to be, resulting in parsing twice, 
 				// once above and once in the Combine() call. Using Combine() here specifically so that special INI syntax (+, ., etc.) is parsed correctly.
-				const FString PluginIniPath = FString::Printf(TEXT("%s%s.ini"), *PluginConfigDir, *PluginIniName);
-				if (ExistingConfig->Combine(PluginIniPath))
+				const FString PluginIniPath = FString::Printf(TEXT("%s%s.ini"), *ConfigDirectory, *PluginIniName);
+				ExistingConfig->CombineFromBuffer(ConfigAsString, PluginIniPath);
+
+#if ALLOW_INI_OVERRIDE_FROM_COMMANDLINE
+				FConfigFile::OverrideFromCommandline(ExistingConfig, Ini.Name);
+#endif // ALLOW_INI_OVERRIDE_FROM_COMMANDLINE
+
+				if (Ini.bCreateDeviceProfiles)
 				{
-					ReloadConfigs(TempConfig);
+					LoadDeviceProfilesFromConfig(Config);
 				}
 				else
 				{
-					UE_LOG(LogGameFeatures, Error, TEXT("[GameFeatureData %s]: Failed to combine INI %s with base INI %s. Aborting import/application of INI settings."), *GetPathNameSafe(this), *PluginIniName, *IniFile);
+					ReloadConfigs(Config);
 				}
 			}
 		}
diff --git a/Engine/Source/Runtime/Engine/Private/DeviceProfiles/DeviceProfileManager.cpp b/Engine/Source/Runtime/Engine/Private/DeviceProfiles/DeviceProfileManager.cpp
--- a/Engine/Source/Runtime/Engine/Private/DeviceProfiles/DeviceProfileManager.cpp
+++ b/Engine/Source/Runtime/Engine/Private/DeviceProfiles/DeviceProfileManager.cpp
@@ -168,7 +168,7 @@ TMap<FName, FString> UDeviceProfileManager::GatherDeviceProfileCVars(const FStri
 
 
 	EPlatformMemorySizeBucket MemBucket = FPlatformMemory::GetMemorySizeBucket();
-	// if caching (for another platfomr), then we use the DP's PreviewMemoryBucket, instead of querying for it
+	// if caching (for another platform), then we use the DP's PreviewMemoryBucket, instead of querying for it
 	if (GatherMode == EDeviceProfileMode::DPM_CacheValues || GatherMode == EDeviceProfileMode::DPM_CacheValuesIgnoreMatchingRules )
 	{
 #if ALLOW_OTHER_PLATFORM_CONFIG
@@ -180,8 +180,15 @@ TMap<FName, FString> UDeviceProfileManager::GatherDeviceProfileCVars(const FStri
 			check(Profile);
 			return DeviceProfileCVars;
 		}
-		// use the DP's platform's configs, NOT the running platform
-		ConfigSystem = FConfigCacheIni::ForPlatform(*Profile->DeviceType);
+
+		// use the DP's platform's configs, NOT the running platform, unless the profile is not there
+		FConfigCacheIni* PlatformConfigSystem = FConfigCacheIni::ForPlatform(*Profile->DeviceType);
+		FConfigFile* PlatformConfigFile = PlatformConfigSystem->FindConfigFile(GDeviceProfilesIni);
+		if (PlatformConfigFile && PlatformConfigFile->Contains(DeviceProfileName))
+		{
+			ConfigSystem = PlatformConfigSystem;
+		}
+
 		MemBucket = Profile->GetPreviewMemorySizeBucket();
 		if(GatherMode == EDeviceProfileMode::DPM_CacheValues)
 		{
