GameFeaturePlugin changes:
- Fixing problem where dependency URLs were cached even if the plugin was filtered as we loaded the plugin details before filtering.
- No longer cache dependency URLs at all for a single plugin, only cache the Dependency PluginName instead of full URLs
- Use ResolvePluginDependency at the point where plugins are actually needed instead of doing it on first call to GetGameFeaturePluginDetails.
- Removed FileExists check in ResolvePluginDependencies as you couldn't get into this function if the uplugin file didn't exist and that was potentially expensive overhead
- Fixing mispelled Reslove to Resolve in ResolvePluginDependencies
- Now filtering any PluginDepedencies that have a GetType of Engine as these can not be GFPs and it is a significant optimization to how many plugins we attempt to resolve / parse from JSON
#tests Windows Dev Client

diff --git a/Engine/Plugins/Experimental/GameFeatures/Source/GameFeatures/Private/GameFeaturesProjectPolicies.cpp b/Engine/Plugins/Experimental/GameFeatures/Source/GameFeatures/Private/GameFeaturesProjectPolicies.cpp
--- a/Engine/Plugins/Experimental/GameFeatures/Source/GameFeatures/Private/GameFeaturesProjectPolicies.cpp
+++ b/Engine/Plugins/Experimental/GameFeatures/Source/GameFeatures/Private/GameFeaturesProjectPolicies.cpp
@@ -49,25 +49,25 @@ bool UGameFeaturesProjectPolicies::WillPluginBeCooked(const FString& PluginFilen
 	return true;
 }
 
-TValueOrError<FString, FString> UGameFeaturesProjectPolicies::ReslovePluginDependency(const FString& PluginURL, const FString& DependencyName) const
+TValueOrError<FString, FString> UGameFeaturesProjectPolicies::ResolvePluginDependency(const FString& PluginURL, const FString& DependencyName) const
 {
 	FString DependencyURL;
 	bool bResolvedDependency = false;
 
+	// Check if UGameFeaturesSubsystem is already aware of it
+	if (UGameFeaturesSubsystem::Get().GetPluginURLByName(DependencyName, DependencyURL))
+	{
+		bResolvedDependency = true;
+	}
 	// Check if the dependency plugin exists yet (should be true for all built-in plugins)
-	if (TSharedPtr<IPlugin> DependencyPlugin = IPluginManager::Get().FindPlugin(DependencyName))
+	else if (TSharedPtr<IPlugin> DependencyPlugin = IPluginManager::Get().FindPlugin(DependencyName))
 	{
-		// Check if UGameFeaturesSubsystem is already aware of it
-		if (!UGameFeaturesSubsystem::Get().GetPluginURLByName(DependencyPlugin->GetName(), DependencyURL))
+		// It could still be a GFP, but state machine may not have been created for it yet
+		// Check if it is a built-in GFP
+		if (!DependencyPlugin->GetDescriptorFileName().IsEmpty() &&
+			GetDefault<UGameFeaturesSubsystemSettings>()->IsValidGameFeaturePlugin(FPaths::ConvertRelativePathToFull(DependencyPlugin->GetDescriptorFileName())))
 		{
-			// It could still be a GFP, but state machine may not have been created for it yet
-			// Check if it is a built-in GFP
-			if (!DependencyPlugin->GetDescriptorFileName().IsEmpty() &&
-				GetDefault<UGameFeaturesSubsystemSettings>()->IsValidGameFeaturePlugin(FPaths::ConvertRelativePathToFull(DependencyPlugin->GetDescriptorFileName())) &&
-				FPaths::FileExists(DependencyPlugin->GetDescriptorFileName()))
-			{
-				DependencyURL = UGameFeaturesSubsystem::GetPluginURL_FileProtocol(DependencyPlugin->GetDescriptorFileName());
-			}
+			DependencyURL = UGameFeaturesSubsystem::GetPluginURL_FileProtocol(DependencyPlugin->GetDescriptorFileName());
 		}
 
 		bResolvedDependency = true;
diff --git a/Engine/Plugins/Experimental/GameFeatures/Source/GameFeatures/Private/GameFeaturesSubsystem.cpp b/Engine/Plugins/Experimental/GameFeatures/Source/GameFeatures/Private/GameFeaturesSubsystem.cpp
--- a/Engine/Plugins/Experimental/GameFeatures/Source/GameFeatures/Private/GameFeaturesSubsystem.cpp
+++ b/Engine/Plugins/Experimental/GameFeatures/Source/GameFeatures/Private/GameFeaturesSubsystem.cpp
@@ -1554,9 +1554,9 @@ void UGameFeaturesSubsystem::LoadBuiltInGameFeaturePlugins(FBuiltInPluginAdditio
 		{
 			LoadContext->Results.Add(Plugin->GetName(), Result);
 			++LoadContext->NumPluginsLoaded;
-			UE_LOG(LogGameFeatures, VeryVerbose, TEXT("Finished Loading %i builtins"), LoadContext->NumPluginsLoaded);
 		}));
 	}
+	UE_LOG(LogGameFeatures, VeryVerbose, TEXT("Finished Loading %i builtins"), LoadContext->NumPluginsLoaded);
 }
 
 bool UGameFeaturesSubsystem::GetPluginURLByName(const FString& PluginName, FString& OutPluginURL) const
@@ -1757,28 +1757,49 @@ bool UGameFeaturesSubsystem::GetGameFeaturePluginDetails(const FString& PluginUR
 				if (ElementObjectPtr && ElementObjectPtr->IsValid())
 				{
 					const TSharedPtr<FJsonObject>& ElementObject = *ElementObjectPtr;
-					bool bElementEnabled = false;
-					ElementObject->TryGetBoolField(EnabledField, bElementEnabled);
 
-					bool bElementActivate = false;
-					ElementObject->TryGetBoolField(ActivateField, bElementActivate);
-
-					if (bElementEnabled)
+					FString DependencyName;
+					ElementObject->TryGetStringField(NameField, DependencyName);
+					if (!DependencyName.IsEmpty())
 					{
-						FString DependencyName;
-						ElementObject->TryGetStringField(NameField, DependencyName);
-						if (!DependencyName.IsEmpty())
+						//parse enabled element from IPlugin data and use that to filter if further parsing is needed
+						bool bElementEnabled = false;
+						if (TSharedPtr<IPlugin> DependencyPlugin = IPluginManager::Get().FindPlugin(DependencyName))
 						{
-							TValueOrError<FString, FString> ResolvedDepResult = GameSpecificPolicies->ReslovePluginDependency(PluginURL, DependencyName);
-							if (ResolvedDepResult.HasError())
+							if (DependencyPlugin->GetType() == EPluginType::Engine)
 							{
-								UE_LOG(LogGameFeatures, Error, TEXT("Game feature plugin '%s' has unknown dependency '%s' [%s]."), *PluginDescriptorFilename, *DependencyName, *ResolvedDepResult.GetError());
+								UE_LOG(LogGameFeatures, VeryVerbose, TEXT("Skipping Dependency %s in %s because it is an engine plugin and thus can not be a GFP"), *DependencyName, *PluginDescriptorFilename);
+								continue;
 							}
-							else if (ResolvedDepResult.HasValue() && !ResolvedDepResult.GetValue().IsEmpty()) // Dependency may not be a GFP
+
+							bElementEnabled = DependencyPlugin->IsEnabled();							
+						}
+						// Plugin is not yet in PluginManager, so need to just go off JSon enabled value
+						else
+						{
+							ElementObject->TryGetBoolField(EnabledField, bElementEnabled);
+							if (bElementEnabled)
 							{
-								OutPluginDetails.PluginDependencies.Emplace(FGameFeaturePluginReferenceDetails(ResolvedDepResult.StealValue(), bElementActivate));
+								UE_LOG(LogGameFeatures, Display, TEXT("Plugin dependency %s marked enabled in %s but not found in PluginManager."), *DependencyName, *PluginDescriptorFilename);
 							}
 						}
+
+						if (bElementEnabled)
+						{
+							//Have to get Activate from JSON as it's unique to GFP and not in the PluginManager
+							bool bElementActivate = false;
+							ElementObject->TryGetBoolField(ActivateField, bElementActivate);
+
+							OutPluginDetails.PluginDependencies.Emplace(FGameFeaturePluginReferenceDetails(DependencyName, bElementActivate));
+						}
+						else
+						{
+							UE_LOG(LogGameFeatures, VeryVerbose, TEXT("Skipping adding dependency %s in %s. Plugin is disabled."), *DependencyName, *PluginDescriptorFilename);
+						}
+					}
+					else
+					{
+						UE_LOG(LogGameFeatures, Error, TEXT("Error parsing dependency name in %s! Invalid JSON data!"), *PluginDescriptorFilename);
 					}
 				}
 			}
@@ -1987,17 +2008,33 @@ bool UGameFeaturesSubsystem::FindOrCreatePluginDependencyStateMachines(const FSt
 	{
 		for (const FGameFeaturePluginReferenceDetails& PluginDependency : Details.PluginDependencies)
 		{
+			const FString& DependencyName = PluginDependency.PluginName;
+			TValueOrError<FString, FString> DependencyURLInfo = GameSpecificPolicies->ResolvePluginDependency(PluginURL, DependencyName);
+			if (DependencyURLInfo.HasError())
+			{
+				UE_LOG(LogGameFeatures, Error, TEXT("Failure to resolve dependency %s [%s] for parent plugin url: %s"), *DependencyName, *DependencyURLInfo.GetError(), *PluginURL);
+				return false;
+			}
+
+			const FString& DependencyURL = DependencyURLInfo.GetValue();
+
+			// Dependency may not be a GFP and so will have an empty URL but not have an error
+			if (DependencyURL.IsEmpty())
+			{
+				continue;
+			}
+
 			// Inherit dep protocol options if possible
 			FGameFeatureProtocolOptions DepProtocolOptions;
-			EGameFeaturePluginProtocol DepProtocol = UGameFeaturesSubsystem::GetPluginURLProtocol(PluginDependency.URL);
+			EGameFeaturePluginProtocol DepProtocol = UGameFeaturesSubsystem::GetPluginURLProtocol(DependencyURL);
 			if (DepProtocol == EGameFeaturePluginProtocol::InstallBundle && InDepProtocolOptions.HasSubtype<FInstallBundlePluginProtocolOptions>())
 			{
 				DepProtocolOptions = InDepProtocolOptions;
 			}
 
-			UGameFeaturePluginStateMachine* Dependency = FindOrCreateGameFeaturePluginStateMachine(PluginDependency.URL, DepProtocolOptions);
-			check(Dependency);
-			OutDependencyMachines.Add(Dependency);
+			UGameFeaturePluginStateMachine* ResolvedDependency = FindOrCreateGameFeaturePluginStateMachine(DependencyURL, DepProtocolOptions);
+			check(ResolvedDependency);
+			OutDependencyMachines.Add(ResolvedDependency);
 		}
 
 		return true;
@@ -2006,7 +2043,7 @@ bool UGameFeaturesSubsystem::FindOrCreatePluginDependencyStateMachines(const FSt
 	return false;
 }
 
-bool UGameFeaturesSubsystem::FindPluginDependencyStateMachinesToActivate(const FString& PluginURL, const FString& PluginFilename, TArray<UGameFeaturePluginStateMachine*>& OutDependencyMachines)
+bool UGameFeaturesSubsystem::FindPluginDependencyStateMachinesToActivate(const FString& PluginURL, const FString& PluginFilename, TArray<UGameFeaturePluginStateMachine*>& OutDependencyMachines) const
 {
 	FGameFeaturePluginDetails Details;
 	if (GetGameFeaturePluginDetails(PluginURL, PluginFilename, Details))
@@ -2015,9 +2052,33 @@ bool UGameFeaturesSubsystem::FindPluginDependencyStateMachinesToActivate(const F
 		{
 			if (PluginDependency.bShouldActivate)
 			{
-				UGameFeaturePluginStateMachine* Dependency = FindGameFeaturePluginStateMachine(PluginDependency.URL);
-				check(Dependency);
-				OutDependencyMachines.Add(Dependency);
+				const FString& DependencyName = PluginDependency.PluginName;
+				TValueOrError<FString, FString> DependencyURLInfo = GameSpecificPolicies->ResolvePluginDependency(PluginURL, DependencyName);
+				if (DependencyURLInfo.HasError())
+				{
+					UE_LOG(LogGameFeatures, Error, TEXT("Failure to resolve dependency %s [%s] for parent plugin url: %s"), *DependencyName, *DependencyURLInfo.GetError(), *PluginURL);
+					return false;
+				}
+
+				const FString& DependencyURL = DependencyURLInfo.GetValue();
+				
+				// Dependency may not be a GFP and so will have an empty URL but not have an error
+				if (DependencyURL.IsEmpty())
+				{
+					continue;
+				}
+
+				UGameFeaturePluginStateMachine* Dependency = FindGameFeaturePluginStateMachine(DependencyURL);
+				if (Dependency)
+				{
+					OutDependencyMachines.Add(Dependency);
+				}
+				//Expect to find all valid dependencies and activate them, so error if not found
+				else
+				{
+					UE_LOG(LogGameFeatures, Error, TEXT("FindPluginDependencyStateMachinesToActivate failed to find plugin state machine for %s using URL %s"), *DependencyName, *DependencyURL);
+					return false;
+				}
 			}
 		}
 
diff --git a/Engine/Plugins/Experimental/GameFeatures/Source/GameFeatures/Public/GameFeaturesProjectPolicies.h b/Engine/Plugins/Experimental/GameFeatures/Source/GameFeatures/Public/GameFeaturesProjectPolicies.h
--- a/Engine/Plugins/Experimental/GameFeatures/Source/GameFeatures/Public/GameFeaturesProjectPolicies.h
+++ b/Engine/Plugins/Experimental/GameFeatures/Source/GameFeatures/Public/GameFeaturesProjectPolicies.h
@@ -44,7 +44,7 @@ public:
 	// Called to resolve plugin dependencies, OutDependencyURL will contain the GFP URL if this dependency is a GFP.
 	// This may be called with file protocol for built-in plugins in some cases, even if a different protocol is used at runtime.
 	// returns The dependency URL or an error if the dependency could not be resolved
-	virtual TValueOrError<FString, FString> ReslovePluginDependency(const FString& PluginURL, const FString& DependencyName) const;
+	virtual TValueOrError<FString, FString> ResolvePluginDependency(const FString& PluginURL, const FString& DependencyName) const;
 
 	// Called by code that explicitly wants to load a specific plugin
 	// (e.g., when using a fast cook a game might want to allow explicitly loaded game feature plugins)
diff --git a/Engine/Plugins/Experimental/GameFeatures/Source/GameFeatures/Public/GameFeaturesSubsystem.h b/Engine/Plugins/Experimental/GameFeatures/Source/GameFeatures/Public/GameFeaturesSubsystem.h
--- a/Engine/Plugins/Experimental/GameFeatures/Source/GameFeatures/Public/GameFeaturesSubsystem.h
+++ b/Engine/Plugins/Experimental/GameFeatures/Source/GameFeatures/Public/GameFeaturesSubsystem.h
@@ -195,11 +195,11 @@ void GAMEFEATURES_API LexFromString(EGameFeatureTargetState& Value, const TCHAR*
 
 struct FGameFeaturePluginReferenceDetails
 {
-	FString URL;
+	FString PluginName;
 	bool bShouldActivate;
 
-	FGameFeaturePluginReferenceDetails(const FString& InURL, bool bInShouldActivate)
-		: URL(InURL)
+	FGameFeaturePluginReferenceDetails(FString InPluginName, bool bInShouldActivate)
+		: PluginName(MoveTemp(InPluginName))
 		, bShouldActivate(bInShouldActivate)
 	{
 	}
@@ -611,7 +611,7 @@ private:
 	template <typename> friend struct FTransitionDependenciesGameFeaturePluginState;
 	friend struct FWaitingForDependenciesTransitionPolicy;
 
-	bool FindPluginDependencyStateMachinesToActivate(const FString& PluginURL, const FString& PluginFilename, TArray<UGameFeaturePluginStateMachine*>& OutDependencyMachines);
+	bool FindPluginDependencyStateMachinesToActivate(const FString& PluginURL, const FString& PluginFilename, TArray<UGameFeaturePluginStateMachine*>& OutDependencyMachines) const;
 	friend struct FActivatingDependenciesTransitionPolicy;
 
 	/** Handle 'ListGameFeaturePlugins' console command */
