Quality of life fix to not ensure when a pak chunk is already in use.

diff --git a/Engine/Plugins/Experimental/GameFeatures/Source/GameFeatures/Private/GameFeatureAction_AddChunkOverride.cpp b/Engine/Plugins/Experimental/GameFeatures/Source/GameFeatures/Private/GameFeatureAction_AddChunkOverride.cpp
--- a/Engine/Plugins/Experimental/GameFeatures/Source/GameFeatures/Private/GameFeatureAction_AddChunkOverride.cpp
+++ b/Engine/Plugins/Experimental/GameFeatures/Source/GameFeatures/Private/GameFeatureAction_AddChunkOverride.cpp
@@ -171,7 +171,8 @@ void UGameFeatureAction_AddChunkOverride::AddChunkIdOverride()
 	}
 	if (GameFeatureAction_AddChunkOverride::ChunkIdToPluginMap.Contains(ChunkId))
 	{
-		ensureMsgf(false, TEXT("ChunkId (%d) is already in use by %s. Manually resolve the conflict"), ChunkId, *GameFeatureAction_AddChunkOverride::ChunkIdToPluginMap[ChunkId]);
+		UE_LOG(LogAddChunkOverride, Error, TEXT("ChunkId (%d) is already in use by %s. Manually resolve the conflict"), ChunkId, *GameFeatureAction_AddChunkOverride::ChunkIdToPluginMap[ChunkId]);
+		FMessageDialog::Open(EAppMsgType::Ok, LOCTEXT("AddChunkOverride_IdConflight", "Chunk Id is already in use"));
 		return;
 	}
 
