Improving error message with with GFP has the duplicate chunkId

diff --git a/Engine/Plugins/Experimental/GameFeatures/Source/GameFeatures/Private/GameFeatureAction_AddChunkOverride.cpp b/Engine/Plugins/Experimental/GameFeatures/Source/GameFeatures/Private/GameFeatureAction_AddChunkOverride.cpp
--- a/Engine/Plugins/Experimental/GameFeatures/Source/GameFeatures/Private/GameFeatureAction_AddChunkOverride.cpp
+++ b/Engine/Plugins/Experimental/GameFeatures/Source/GameFeatures/Private/GameFeatureAction_AddChunkOverride.cpp
@@ -171,8 +171,13 @@ void UGameFeatureAction_AddChunkOverride::AddChunkIdOverride()
 	}
 	if (GameFeatureAction_AddChunkOverride::ChunkIdToPluginMap.Contains(ChunkId))
 	{
-		UE_LOG(LogAddChunkOverride, Error, TEXT("ChunkId (%d) is already in use by %s. Manually resolve the conflict"), ChunkId, *GameFeatureAction_AddChunkOverride::ChunkIdToPluginMap[ChunkId]);
-		FMessageDialog::Open(EAppMsgType::Ok, LOCTEXT("AddChunkOverride_IdConflight", "Chunk Id is already in use"));
+		FString PluginName;
+		if (UGameFeatureData* GameFeatureData = GetTypedOuter<UGameFeatureData>())
+		{
+			GameFeatureData->GetPluginName(PluginName);
+		}
+		UE_LOG(LogAddChunkOverride, Error, TEXT("ChunkId (%d) is already in use by %s. Manually resolve the conflict for %s"), ChunkId, *GameFeatureAction_AddChunkOverride::ChunkIdToPluginMap[ChunkId], *PluginName);
+		FMessageDialog::Open(EAppMsgType::Ok, FText::Format(LOCTEXT("AddChunkOverride_IdConflight", "Chunk Id is already in use by '{0}'."), FText::FromString(PluginName)));
 		return;
 	}
 
