Making it so having an unknown plugin in your GameFeaturePlugin dependencies list will not cause a WaitingForDependenciesError and will only log an error. This is to mirror behavior before a previous change to how dependencies are cached.
Undoing previous logging change with some other minor log changes.
#tests Windows Dev Client

diff --git a/Engine/Plugins/Experimental/GameFeatures/Source/GameFeatures/Private/GameFeaturesSubsystem.cpp b/Engine/Plugins/Experimental/GameFeatures/Source/GameFeatures/Private/GameFeaturesSubsystem.cpp
--- a/Engine/Plugins/Experimental/GameFeatures/Source/GameFeatures/Private/GameFeaturesSubsystem.cpp
+++ b/Engine/Plugins/Experimental/GameFeatures/Source/GameFeatures/Private/GameFeaturesSubsystem.cpp
@@ -1554,9 +1554,9 @@ void UGameFeaturesSubsystem::LoadBuiltInGameFeaturePlugins(FBuiltInPluginAdditio
 		{
 			LoadContext->Results.Add(Plugin->GetName(), Result);
 			++LoadContext->NumPluginsLoaded;
+			UE_LOG(LogGameFeatures, VeryVerbose, TEXT("Finished Loading %i builtins"), LoadContext->NumPluginsLoaded);
 		}));
 	}
-	UE_LOG(LogGameFeatures, VeryVerbose, TEXT("Finished Loading %i builtins"), LoadContext->NumPluginsLoaded);
 }
 
 bool UGameFeaturesSubsystem::GetPluginURLByName(const FString& PluginName, FString& OutPluginURL) const
@@ -1780,7 +1780,7 @@ bool UGameFeaturesSubsystem::GetGameFeaturePluginDetails(const FString& PluginUR
 							ElementObject->TryGetBoolField(EnabledField, bElementEnabled);
 							if (bElementEnabled)
 							{
-								UE_LOG(LogGameFeatures, Display, TEXT("Plugin dependency %s marked enabled in %s but not found in PluginManager."), *DependencyName, *PluginDescriptorFilename);
+								UE_LOG(LogGameFeatures, Verbose, TEXT("Plugin dependency %s marked enabled in %s but not found in PluginManager."), *DependencyName, *PluginDescriptorFilename);
 							}
 						}
 
@@ -2012,8 +2012,12 @@ bool UGameFeaturesSubsystem::FindOrCreatePluginDependencyStateMachines(const FSt
 			TValueOrError<FString, FString> DependencyURLInfo = GameSpecificPolicies->ResolvePluginDependency(PluginURL, DependencyName);
 			if (DependencyURLInfo.HasError())
 			{
-				UE_LOG(LogGameFeatures, Error, TEXT("Failure to resolve dependency %s [%s] for parent plugin url: %s"), *DependencyName, *DependencyURLInfo.GetError(), *PluginURL);
-				return false;
+				UE_LOG(LogGameFeatures, Error, TEXT("Game feature plugin '%s' has unknown dependency '%s' [%s]."), *PluginFilename, *DependencyName, *DependencyURLInfo.GetError());
+
+				//Don't actually return false here as we want to still be able to progress in the case of 
+				//things like an editor plugin being included as a dependency in the client or a dynamic dependency that
+				//hasn't correctly loaded yet
+				continue;
 			}
 
 			const FString& DependencyURL = DependencyURLInfo.GetValue();
