- Add GFP observer pre-mounting callback
- Add GFP observer post-mounting callback that provides pausing capability
#rnx

diff --git a/Engine/Plugins/Experimental/GameFeatures/Source/GameFeatures/Private/GameFeaturePluginStateMachine.cpp b/Engine/Plugins/Experimental/GameFeatures/Source/GameFeatures/Private/GameFeaturePluginStateMachine.cpp
--- a/Engine/Plugins/Experimental/GameFeatures/Source/GameFeatures/Private/GameFeaturePluginStateMachine.cpp
+++ b/Engine/Plugins/Experimental/GameFeatures/Source/GameFeatures/Private/GameFeaturePluginStateMachine.cpp
@@ -686,7 +686,7 @@ struct FGameFeaturePluginState_Terminal : public FDestinationGameFeaturePluginSt
 		checkf(!bEnteredTerminalState, TEXT("Plugin entered terminal state more than once! %s"), *StateProperties.PluginIdentifier.GetFullPluginURL());
 		bEnteredTerminalState = true;
 
-		UGameFeaturesSubsystem::Get().OnGameFeatureTerminating(StateProperties.PluginName, StateProperties.PluginIdentifier.GetFullPluginURL());
+		UGameFeaturesSubsystem::Get().OnGameFeatureTerminating(StateProperties.PluginName, StateProperties.PluginIdentifier);
 	}
 };
 
@@ -704,7 +704,7 @@ struct FGameFeaturePluginState_UnknownStatus : public FDestinationGameFeaturePlu
 		{
 			StateStatus.SetTransition(EGameFeaturePluginState::CheckingStatus);
 
-			UGameFeaturesSubsystem::Get().OnGameFeatureCheckingStatus(StateProperties.PluginIdentifier.GetFullPluginURL());
+			UGameFeaturesSubsystem::Get().OnGameFeatureCheckingStatus(StateProperties.PluginIdentifier);
 		}
 	}
 };
@@ -784,7 +784,7 @@ struct FGameFeaturePluginState_CheckingStatus : public FGameFeaturePluginState
 			return;
 		}
 
-		UGameFeaturesSubsystem::Get().OnGameFeatureStatusKnown(StateProperties.PluginName, StateProperties.PluginIdentifier.GetFullPluginURL());
+		UGameFeaturesSubsystem::Get().OnGameFeatureStatusKnown(StateProperties.PluginName, StateProperties.PluginIdentifier);
 		StateStatus.SetTransition(EGameFeaturePluginState::StatusKnown);
 	}
 };
@@ -1370,7 +1370,7 @@ struct FGameFeaturePluginState_Downloading : public FGameFeaturePluginState
 	void NotifyPauseChange(bool bIsPaused, FString PauseReason)
 	{
 		FGameFeaturePauseStateChangeContext Context(LexToString(EGameFeaturePluginState::Downloading), PauseReason, bIsPaused);
-		UGameFeaturesSubsystem::Get().OnGameFeaturePauseChange(StateProperties.PluginIdentifier.GetFullPluginURL(), StateProperties.PluginName, Context);
+		UGameFeaturesSubsystem::Get().OnGameFeaturePauseChange(StateProperties.PluginIdentifier, StateProperties.PluginName, Context);
 	}
 
 	virtual void BeginState() override
@@ -1675,9 +1675,10 @@ struct FGameFeaturePluginState_Mounting : public FGameFeaturePluginState
 		, Result(MakeValue())
 	{}
 
-	UE::GameFeatures::FResult Result;
+	int32 NumObservedPostMountPausers = 0;
+	int32 NumExpectedPostMountPausers = 0;
 	TArray<FName> PendingBundles;
-	bool bMounted = false;
+	UE::GameFeatures::FResult Result;
 
 	void OnInstallBundleCompleted(FInstallBundleRequestResultInfo BundleResult)
 	{
@@ -1700,17 +1701,13 @@ struct FGameFeaturePluginState_Mounting : public FGameFeaturePluginState
 			}
 		}
 
-		if (PendingBundles.Num() > 0)
-		{
-			return;
-		}
-
-		if (Result.HasValue())
+		if (PendingBundles.IsEmpty())
 		{
-			bMounted = true;
+			IInstallBundleManager::InstallBundleCompleteDelegate.RemoveAll(this);
+			FCoreDelegates::GetOnPakFileMounted2().RemoveAll(this);
+			
+			UpdateStateMachineImmediate();
 		}
-
-		UpdateStateMachineImmediate();
 	}
 
 	void OnPakFileMounted(const IPakFile& PakFile)
@@ -1728,15 +1725,29 @@ struct FGameFeaturePluginState_Mounting : public FGameFeaturePluginState
 		}
 	}
 
+	void OnPostMountPauserCompleted(FStringView InPauserTag)
+	{
+		check(IsInGameThread());
+		ensure(NumExpectedPostMountPausers != INDEX_NONE);
+		++NumObservedPostMountPausers;
+
+		UE_LOG(LogGameFeatures, Display, TEXT("Post-mount of %s resumed by %.*s"), *StateProperties.PluginName, InPauserTag.Len(), InPauserTag.GetData());
+
+		if (NumObservedPostMountPausers == NumExpectedPostMountPausers)
+		{
+			UpdateStateMachineDeferred();
+		}
+	}
+
 	virtual void BeginState() override
 	{
-		Result = MakeValue();
+		NumObservedPostMountPausers = 0;
+		NumExpectedPostMountPausers = 0;
 		PendingBundles.Empty();
-		bMounted = false;
+		Result = MakeValue();
 
 		if (StateProperties.GetPluginProtocol() != EGameFeaturePluginProtocol::InstallBundle)
 		{
-			bMounted = true;
 			return;
 		}
 		
@@ -1771,11 +1782,7 @@ struct FGameFeaturePluginState_Mounting : public FGameFeaturePluginState
 			return;
 		}
 
-		if (RequestInfo.BundlesEnqueued.Num() == 0)
-		{
-			bMounted = true;
-		}
-		else
+		if (RequestInfo.BundlesEnqueued.Num() > 0)
 		{
 			PendingBundles = MoveTemp(RequestInfo.BundlesEnqueued);
 			IInstallBundleManager::InstallBundleCompleteDelegate.AddRaw(this, &FGameFeaturePluginState_Mounting::OnInstallBundleCompleted);
@@ -1789,74 +1796,132 @@ struct FGameFeaturePluginState_Mounting : public FGameFeaturePluginState
 	virtual void UpdateState(FGameFeaturePluginStateStatus& StateStatus) override
 	{
 		TRACE_CPUPROFILER_EVENT_SCOPE(GFP_Mounting);
-		if (!Result.HasValue())
-		{
-			StateStatus.SetTransitionError(EGameFeaturePluginState::ErrorMounting, Result);
-			return;
-		}
-		if (!bMounted)
+
+		// Check if waiting for install bundles
+		if (PendingBundles.Num() > 0)
 		{
 			return;
 		}
 
-		checkf(!StateProperties.PluginInstalledFilename.IsEmpty(), TEXT("PluginInstalledFilename must be set by the Mounting. PluginURL: %s"), *StateProperties.PluginIdentifier.GetFullPluginURL());
-		checkf(FPaths::GetExtension(StateProperties.PluginInstalledFilename) == TEXT("uplugin"), TEXT("PluginInstalledFilename must have a uplugin extension. PluginURL: %s"), *StateProperties.PluginIdentifier.GetFullPluginURL());
-
-		// refresh the plugins list to let the plugin manager know about it
-		const TSharedPtr<IPlugin> MaybePlugin = IPluginManager::Get().FindPlugin(StateProperties.PluginName);
-		const bool bNeedsPluginMount = (MaybePlugin == nullptr || MaybePlugin->GetDescriptor().bExplicitlyLoaded);
-
-		if (MaybePlugin)
+		// Check if post-mount is paused
+		if (NumExpectedPostMountPausers > 0)
 		{
-			if (!FPaths::IsSamePath(MaybePlugin->GetDescriptorFileName(), StateProperties.PluginInstalledFilename))
+			// Check if post-mount unpaused
+			if (NumExpectedPostMountPausers == NumObservedPostMountPausers)
 			{
-				StateStatus.SetTransitionError(EGameFeaturePluginState::ErrorMounting, GetErrorResult(TEXT("Plugin_Name_Already_In_Use")));
-				return;
+				NumExpectedPostMountPausers = INDEX_NONE;
+
+				// We previously sent an OnGameFeaturePauseChange delegate we need to send that work is now unpaused
+				FGameFeaturePauseStateChangeContext UnPauseContext(LexToString(EGameFeaturePluginState::Mounting), TEXT(""), /*bIsPausedIn=*/false);
+				UGameFeaturesSubsystem::Get().OnGameFeaturePauseChange(StateProperties.PluginIdentifier, StateProperties.PluginName, UnPauseContext);
+
+				TransitionOut(StateStatus);
 			}
+			return;
 		}
-		else
+
+		// Don't mount the plugin if there was an error during BeginState or bundles install
+		if (!Result.HasError())
 		{
-			const bool bAddedPlugin = IPluginManager::Get().AddToPluginsList(StateProperties.PluginInstalledFilename);
-			if (!bAddedPlugin)
+			// Pre-mount
 			{
-				StateStatus.SetTransitionError(EGameFeaturePluginState::ErrorMounting, GetErrorResult(TEXT("Failed_To_Register_Plugin")));
-				return;
+				FGameFeaturePreMountingContext Context;
+				UGameFeaturesSubsystem::Get().OnGameFeaturePreMounting(StateProperties.PluginName, StateProperties.PluginIdentifier, Context);
 			}
 
-			StateProperties.bAddedPluginToManager = true;
-		}
-		
-		if (bNeedsPluginMount)
-		{
-			IPluginManager::Get().MountExplicitlyLoadedPlugin(StateProperties.PluginName);
-		}
+			checkf(!StateProperties.PluginInstalledFilename.IsEmpty(), TEXT("PluginInstalledFilename must be set by the Mounting. PluginURL: %s"), *StateProperties.PluginIdentifier.GetFullPluginURL());
+			checkf(FPaths::GetExtension(StateProperties.PluginInstalledFilename) == TEXT("uplugin"), TEXT("PluginInstalledFilename must have a uplugin extension. PluginURL: %s"), *StateProperties.PluginIdentifier.GetFullPluginURL());
 
-		// After the new plugin is mounted add the asset registry for that plugin.
-		if (StateProperties.GetPluginProtocol() == EGameFeaturePluginProtocol::InstallBundle)
-		{
-			const FString PluginFolder = FPaths::GetPath(StateProperties.PluginInstalledFilename);
-			const FString PluginAssetRegistry = PluginFolder / TEXT("AssetRegistry.bin");
+			// refresh the plugins list to let the plugin manager know about it
+			const TSharedPtr<IPlugin> MaybePlugin = IPluginManager::Get().FindPlugin(StateProperties.PluginName);
+			const bool bNeedsPluginMount = (MaybePlugin == nullptr || MaybePlugin->GetDescriptor().bExplicitlyLoaded);
 
-			TSharedPtr<IPlugin> NewlyMountedPlugin = IPluginManager::Get().FindPlugin(StateProperties.PluginName);
-			if (NewlyMountedPlugin.IsValid() && NewlyMountedPlugin->CanContainContent() && IFileManager::Get().FileExists(*PluginAssetRegistry))
+			if (MaybePlugin)
+			{
+				if (!FPaths::IsSamePath(MaybePlugin->GetDescriptorFileName(), StateProperties.PluginInstalledFilename))
+				{
+					Result = GetErrorResult(TEXT("Plugin_Name_Already_In_Use"));
+				}
+			}
+			else
 			{
-				TArray<uint8> SerializedAssetData;
-				if (!FFileHelper::LoadFileToArray(SerializedAssetData, *PluginAssetRegistry))
+				const bool bAddedPlugin = IPluginManager::Get().AddToPluginsList(StateProperties.PluginInstalledFilename);
+				if (bAddedPlugin)
 				{
-					StateStatus.SetTransitionError(EGameFeaturePluginState::ErrorMounting, GetErrorResult(TEXT("Failed_To_Load_Plugin_AssetRegistry")));
-					return;
+					StateProperties.bAddedPluginToManager = true;
 				}
+				else
+				{
+					Result = GetErrorResult(TEXT("Failed_To_Register_Plugin"));
+				}
+			}
 
-				FAssetRegistryState PluginAssetRegistryState;
-				FMemoryReader Ar(SerializedAssetData);
-				PluginAssetRegistryState.Load(Ar);
+			if (!Result.HasError())
+			{
+				if (bNeedsPluginMount)
+				{
+					IPluginManager::Get().MountExplicitlyLoadedPlugin(StateProperties.PluginName);
+				}
 
-				IAssetRegistry& AssetRegistry = UAssetManager::Get().GetAssetRegistry();
-				AssetRegistry.AppendState(PluginAssetRegistryState);
+				// After the new plugin is mounted add the asset registry for that plugin.
+				if (StateProperties.GetPluginProtocol() == EGameFeaturePluginProtocol::InstallBundle)
+				{
+					const FString PluginFolder = FPaths::GetPath(StateProperties.PluginInstalledFilename);
+					const FString PluginAssetRegistry = PluginFolder / TEXT("AssetRegistry.bin");
+
+					TSharedPtr<IPlugin> NewlyMountedPlugin = IPluginManager::Get().FindPlugin(StateProperties.PluginName);
+					if (NewlyMountedPlugin.IsValid() && NewlyMountedPlugin->CanContainContent() && IFileManager::Get().FileExists(*PluginAssetRegistry))
+					{
+						TArray<uint8> SerializedAssetData;
+						if (FFileHelper::LoadFileToArray(SerializedAssetData, *PluginAssetRegistry))
+						{
+							FAssetRegistryState PluginAssetRegistryState;
+							FMemoryReader Ar(SerializedAssetData);
+							PluginAssetRegistryState.Load(Ar);
+
+							IAssetRegistry& AssetRegistry = UAssetManager::Get().GetAssetRegistry();
+							AssetRegistry.AppendState(PluginAssetRegistryState);
+						}
+						else
+						{
+							Result = GetErrorResult(TEXT("Failed_To_Load_Plugin_AssetRegistry"));
+						}
+					}
+				}
+			}
+		}
+
+		// Post-mount
+		{
+			FGameFeaturePostMountingContext Context(StateProperties.PluginName, [this](FStringView InPauserTag) { OnPostMountPauserCompleted(InPauserTag); });
+			NumExpectedPostMountPausers = INDEX_NONE;
+			UGameFeaturesSubsystem::Get().OnGameFeaturePostMounting(StateProperties.PluginName, StateProperties.PluginIdentifier, Context);
+			NumExpectedPostMountPausers = Context.NumPausers;
+
+			// Check if we got post-mount paused
+			if (NumExpectedPostMountPausers > 0)
+			{
+				// Since we are pausing work during this mounting, also notify the OnGameFeaturePauseChange delegate
+				FGameFeaturePauseStateChangeContext PauseContext(LexToString(EGameFeaturePluginState::Mounting), TEXT("PendingPostMountCallbacks"), true);
+				UGameFeaturesSubsystem::Get().OnGameFeaturePauseChange(StateProperties.PluginIdentifier, StateProperties.PluginName, PauseContext);
+			}
+			else
+			{
+				TransitionOut(StateStatus);
 			}
 		}
+	}
 
-		StateStatus.SetTransition(EGameFeaturePluginState::WaitingForDependencies);
+	void TransitionOut(FGameFeaturePluginStateStatus& StateStatus)
+	{
+		if (Result.HasError())
+		{
+			StateStatus.SetTransitionError(EGameFeaturePluginState::ErrorMounting, Result);
+		}
+		else
+		{
+			StateStatus.SetTransition(EGameFeaturePluginState::WaitingForDependencies);
+		}
 	}
 
 	virtual void EndState() override
@@ -1936,7 +2001,7 @@ struct FGameFeaturePluginState_Unregistering : public FGameFeaturePluginState
 
 		if (StateProperties.GameFeatureData)
 		{
-			UGameFeaturesSubsystem::Get().OnGameFeatureUnregistering(StateProperties.GameFeatureData, StateProperties.PluginName, StateProperties.PluginIdentifier.GetFullPluginURL());
+			UGameFeaturesSubsystem::Get().OnGameFeatureUnregistering(StateProperties.GameFeatureData, StateProperties.PluginName, StateProperties.PluginIdentifier);
 
 			UGameFeaturesSubsystem::RemoveGameFeatureFromAssetManager(StateProperties.GameFeatureData, StateProperties.PluginName, StateProperties.AddedPrimaryAssetTypes);
 			StateProperties.AddedPrimaryAssetTypes.Empty();
@@ -2052,7 +2117,7 @@ struct FGameFeaturePluginState_Registering : public FGameFeaturePluginState
 			check(StateProperties.AddedPrimaryAssetTypes.Num() == 0);
 			UGameFeaturesSubsystem::Get().AddGameFeatureToAssetManager(StateProperties.GameFeatureData, StateProperties.PluginName, StateProperties.AddedPrimaryAssetTypes);
 
-			UGameFeaturesSubsystem::Get().OnGameFeatureRegistering(StateProperties.GameFeatureData, StateProperties.PluginName, StateProperties.PluginIdentifier.GetFullPluginURL());
+			UGameFeaturesSubsystem::Get().OnGameFeatureRegistering(StateProperties.GameFeatureData, StateProperties.PluginName, StateProperties.PluginIdentifier);
 		}
 		else
 		{
@@ -2085,7 +2150,7 @@ struct FGameFeaturePluginState_Unloading : public FGameFeaturePluginState
 
 	bool bHasUnloaded = false;
 
-	virtual void BeginState() 
+	virtual void BeginState() override
 	{
 		bHasUnloaded = false;
 	}
@@ -2156,7 +2221,7 @@ struct FGameFeaturePluginState_Loading : public FGameFeaturePluginState
 			BundleHandle->WaitUntilComplete(0.0f, false);
 		}
 
-		UGameFeaturesSubsystem::Get().OnGameFeatureLoading(StateProperties.GameFeatureData, StateProperties.PluginIdentifier.GetFullPluginURL());
+		UGameFeaturesSubsystem::Get().OnGameFeatureLoading(StateProperties.GameFeatureData, StateProperties.PluginIdentifier);
 
 		StateStatus.SetTransition(EGameFeaturePluginState::Loaded);
 	}
@@ -2252,14 +2317,14 @@ struct FGameFeaturePluginState_Deactivating : public FGameFeaturePluginState
 
 			// Deactivate
 			FGameFeatureDeactivatingContext Context(StateProperties.PluginName, [this](FStringView InPauserTag) { OnPauserCompleted(InPauserTag); });
-			UGameFeaturesSubsystem::Get().OnGameFeatureDeactivating(StateProperties.GameFeatureData, StateProperties.PluginName, Context, StateProperties.PluginIdentifier.GetFullPluginURL());
+			UGameFeaturesSubsystem::Get().OnGameFeatureDeactivating(StateProperties.GameFeatureData, StateProperties.PluginName, Context, StateProperties.PluginIdentifier);
 			NumExpectedPausers = Context.NumPausers;
 
 			// Since we are pausing work during this deactivation, also notify the OnGameFeaturePauseChange delegate
 			if (NumExpectedPausers > 0)
 			{
 				FGameFeaturePauseStateChangeContext PauseContext(LexToString(EGameFeaturePluginState::Deactivating), TEXT("PendingDeactivationCallbacks"), true);
-				UGameFeaturesSubsystem::Get().OnGameFeaturePauseChange(StateProperties.PluginIdentifier.GetFullPluginURL(), StateProperties.PluginName, PauseContext);
+				UGameFeaturesSubsystem::Get().OnGameFeaturePauseChange(StateProperties.PluginIdentifier, StateProperties.PluginName, PauseContext);
 			}
 		}
 
@@ -2269,7 +2334,7 @@ struct FGameFeaturePluginState_Deactivating : public FGameFeaturePluginState
 			if (NumExpectedPausers > 0)
 			{
 				FGameFeaturePauseStateChangeContext PauseContext(LexToString(EGameFeaturePluginState::Deactivating), TEXT(""), false);
-				UGameFeaturesSubsystem::Get().OnGameFeaturePauseChange(StateProperties.PluginIdentifier.GetFullPluginURL(), StateProperties.PluginName, PauseContext);
+				UGameFeaturesSubsystem::Get().OnGameFeaturePauseChange(StateProperties.PluginIdentifier, StateProperties.PluginName, PauseContext);
 			}
 
 			if (!bHasUnloaded && StateProperties.Destination.MaxState == EGameFeaturePluginState::Loaded)
@@ -2360,7 +2425,7 @@ struct FGameFeaturePluginState_Activating : public FGameFeaturePluginState
 			StateProperties.GameFeatureData->InitializeHierarchicalPluginIniFiles(StateProperties.PluginInstalledFilename);
 		}
 
-		UGameFeaturesSubsystem::Get().OnGameFeatureActivating(StateProperties.GameFeatureData, StateProperties.PluginName, Context, StateProperties.PluginIdentifier.GetFullPluginURL());
+		UGameFeaturesSubsystem::Get().OnGameFeatureActivating(StateProperties.GameFeatureData, StateProperties.PluginName, Context, StateProperties.PluginIdentifier);
 
 		// If this plugin caused localization data to load, wait for that here before marking it as active
 		if (TSharedPtr<IPlugin> Plugin = IPluginManager::Get().FindPlugin(StateProperties.PluginName);
diff --git a/Engine/Plugins/Experimental/GameFeatures/Source/GameFeatures/Private/GameFeaturesSubsystem.cpp b/Engine/Plugins/Experimental/GameFeatures/Source/GameFeatures/Private/GameFeaturesSubsystem.cpp
--- a/Engine/Plugins/Experimental/GameFeatures/Source/GameFeatures/Private/GameFeaturesSubsystem.cpp
+++ b/Engine/Plugins/Experimental/GameFeatures/Source/GameFeatures/Private/GameFeaturesSubsystem.cpp
@@ -249,7 +249,17 @@ FSimpleDelegate FGameFeatureDeactivatingContext::PauseDeactivationUntilComplete(
 
 	++NumPausers;
 	return FSimpleDelegate::CreateLambda(
-		[CompletionCallback=CompletionCallback, PauserTag=MoveTemp(InPauserTag)]() { CompletionCallback(PauserTag); }
+		[CompletionCallback = CompletionCallback, PauserTag = MoveTemp(InPauserTag)]() { CompletionCallback(PauserTag); }
+	);
+}
+
+FSimpleDelegate FGameFeaturePostMountingContext::PauseUntilComplete(FString InPauserTag)
+{
+	UE_LOG(LogGameFeatures, Display, TEXT("Post-mount of %.*s paused by %s"), PluginName.Len(), PluginName.GetData(), *InPauserTag);
+
+	++NumPausers;
+	return FSimpleDelegate::CreateLambda(
+		[CompletionCallback = CompletionCallback, PauserTag = MoveTemp(InPauserTag)]() { CompletionCallback(PauserTag); }
 	);
 }
 
@@ -693,9 +703,9 @@ bool UGameFeaturesSubsystem::IsPluginURLProtocol(FStringView PluginURL, EGameFea
 	return PluginURL.StartsWith(UE::GameFeatures::GameFeaturePluginProtocolPrefix(PluginProtocol));
 }
 
-void UGameFeaturesSubsystem::OnGameFeatureTerminating(const FString& PluginName, const FString& PluginURL)
+void UGameFeaturesSubsystem::OnGameFeatureTerminating(const FString& PluginName, const FGameFeaturePluginIdentifier& PluginIdentifier)
 {
-	CallbackObservers(EObserverCallback::Terminating, PluginURL, &PluginName);
+	CallbackObservers(EObserverCallback::Terminating, PluginIdentifier, &PluginName);
 
 	if (!PluginName.IsEmpty())
 	{
@@ -704,23 +714,33 @@ void UGameFeaturesSubsystem::OnGameFeatureTerminating(const FString& PluginName,
 	}
 }
 
-void UGameFeaturesSubsystem::OnGameFeatureCheckingStatus(const FString& PluginURL)
+void UGameFeaturesSubsystem::OnGameFeatureCheckingStatus(const FGameFeaturePluginIdentifier& PluginIdentifier)
 {
-	CallbackObservers(EObserverCallback::CheckingStatus, PluginURL);
+	CallbackObservers(EObserverCallback::CheckingStatus, PluginIdentifier);
 }
 
-void UGameFeaturesSubsystem::OnGameFeatureStatusKnown(const FString& PluginName, const FString& PluginURL)
+void UGameFeaturesSubsystem::OnGameFeatureStatusKnown(const FString& PluginName, const FGameFeaturePluginIdentifier& PluginIdentifier)
 {
 	// Map plugin name to plugin URL
 	if (ensure(!GameFeaturePluginNameToPathMap.Contains(PluginName)))
 	{
-		GameFeaturePluginNameToPathMap.Add(PluginName, PluginURL);
+		GameFeaturePluginNameToPathMap.Add(PluginName, PluginIdentifier.GetFullPluginURL());
 	}
 }
 
-void UGameFeaturesSubsystem::OnGameFeatureRegistering(const UGameFeatureData* GameFeatureData, const FString& PluginName, const FString& PluginURL)
+void UGameFeaturesSubsystem::OnGameFeaturePreMounting(const FString& PluginName, const FGameFeaturePluginIdentifier& PluginIdentifier, FGameFeaturePreMountingContext& Context)
 {
-	CallbackObservers(EObserverCallback::Registering, PluginURL, &PluginName, GameFeatureData);
+	CallbackObservers(EObserverCallback::PreMounting, PluginIdentifier, &PluginName, /*GameFeatureData=*/nullptr, &Context);
+}
+
+void UGameFeaturesSubsystem::OnGameFeaturePostMounting(const FString& PluginName, const FGameFeaturePluginIdentifier& PluginIdentifier, FGameFeaturePostMountingContext& Context)
+{
+	CallbackObservers(EObserverCallback::PostMounting, PluginIdentifier, &PluginName, /*GameFeatureData=*/nullptr, &Context);
+}
+
+void UGameFeaturesSubsystem::OnGameFeatureRegistering(const UGameFeatureData* GameFeatureData, const FString& PluginName, const FGameFeaturePluginIdentifier& PluginIdentifier)
+{
+	CallbackObservers(EObserverCallback::Registering, PluginIdentifier, &PluginName, GameFeatureData);
 
 	for (UGameFeatureAction* Action : GameFeatureData->GetActions())
 	{
@@ -731,9 +751,9 @@ void UGameFeaturesSubsystem::OnGameFeatureRegistering(const UGameFeatureData* Ga
 	}
 }
 
-void UGameFeaturesSubsystem::OnGameFeatureUnregistering(const UGameFeatureData* GameFeatureData, const FString& PluginName, const FString& PluginURL)
+void UGameFeaturesSubsystem::OnGameFeatureUnregistering(const UGameFeatureData* GameFeatureData, const FString& PluginName, const FGameFeaturePluginIdentifier& PluginIdentifier)
 {
-	CallbackObservers(EObserverCallback::Unregistering, PluginURL, &PluginName, GameFeatureData);
+	CallbackObservers(EObserverCallback::Unregistering, PluginIdentifier, &PluginName, GameFeatureData);
 
 	for (UGameFeatureAction* Action : GameFeatureData->GetActions())
 	{
@@ -744,9 +764,9 @@ void UGameFeaturesSubsystem::OnGameFeatureUnregistering(const UGameFeatureData*
 	}
 }
 
-void UGameFeaturesSubsystem::OnGameFeatureLoading(const UGameFeatureData* GameFeatureData, const FString& PluginURL)
+void UGameFeaturesSubsystem::OnGameFeatureLoading(const UGameFeatureData* GameFeatureData, const FGameFeaturePluginIdentifier& PluginIdentifier)
 {
-	CallbackObservers(EObserverCallback::Loading, PluginURL, nullptr, GameFeatureData);
+	CallbackObservers(EObserverCallback::Loading, PluginIdentifier, nullptr, GameFeatureData);
 
 	for (UGameFeatureAction* Action : GameFeatureData->GetActions())
 	{
@@ -757,9 +777,9 @@ void UGameFeaturesSubsystem::OnGameFeatureLoading(const UGameFeatureData* GameFe
 	}
 }
 
-void UGameFeaturesSubsystem::OnGameFeatureActivating(const UGameFeatureData* GameFeatureData, const FString& PluginName, FGameFeatureActivatingContext& Context, const FString& PluginURL)
+void UGameFeaturesSubsystem::OnGameFeatureActivating(const UGameFeatureData* GameFeatureData, const FString& PluginName, FGameFeatureActivatingContext& Context, const FGameFeaturePluginIdentifier& PluginIdentifier)
 {
-	CallbackObservers(EObserverCallback::Activating, PluginURL, &PluginName, GameFeatureData);
+	CallbackObservers(EObserverCallback::Activating, PluginIdentifier, &PluginName, GameFeatureData);
 
 	for (UGameFeatureAction* Action : GameFeatureData->GetActions())
 	{
@@ -770,9 +790,9 @@ void UGameFeaturesSubsystem::OnGameFeatureActivating(const UGameFeatureData* Gam
 	}
 }
 
-void UGameFeaturesSubsystem::OnGameFeatureDeactivating(const UGameFeatureData* GameFeatureData, const FString& PluginName, FGameFeatureDeactivatingContext& Context, const FString& PluginURL)
+void UGameFeaturesSubsystem::OnGameFeatureDeactivating(const UGameFeatureData* GameFeatureData, const FString& PluginName, FGameFeatureDeactivatingContext& Context, const FGameFeaturePluginIdentifier& PluginIdentifier)
 {
-	CallbackObservers(EObserverCallback::Deactivating, PluginURL, &PluginName, GameFeatureData, &Context);
+	CallbackObservers(EObserverCallback::Deactivating, PluginIdentifier, &PluginName, GameFeatureData, &Context);
 
 	for (UGameFeatureAction* Action : GameFeatureData->GetActions())
 	{
@@ -783,9 +803,9 @@ void UGameFeaturesSubsystem::OnGameFeatureDeactivating(const UGameFeatureData* G
 	}
 }
 
-void UGameFeaturesSubsystem::OnGameFeaturePauseChange(const FString& PluginURL, const FString& PluginName, FGameFeaturePauseStateChangeContext& Context)
+void UGameFeaturesSubsystem::OnGameFeaturePauseChange(const FGameFeaturePluginIdentifier& PluginIdentifier, const FString& PluginName, FGameFeaturePauseStateChangeContext& Context)
 {
-	CallbackObservers(EObserverCallback::PauseChanged, PluginURL, &PluginName, nullptr, &Context);
+	CallbackObservers(EObserverCallback::PauseChanged, PluginIdentifier, &PluginName, nullptr, &Context);
 }
 
 const UGameFeatureData* UGameFeaturesSubsystem::GetDataForStateMachine(UGameFeaturePluginStateMachine* GFSM) const
@@ -841,6 +861,15 @@ bool UGameFeaturesSubsystem::IsGameFeaturePluginInstalled(const FString& PluginU
 	return false;
 }
 
+bool UGameFeaturesSubsystem::IsGameFeaturePluginMounted(const FString& PluginURL) const
+{
+	if (const UGameFeaturePluginStateMachine* StateMachine = FindGameFeaturePluginStateMachine(PluginURL))
+	{
+		return StateMachine->GetCurrentState() > EGameFeaturePluginState::Mounting;
+	}
+	return false;
+}
+
 bool UGameFeaturesSubsystem::IsGameFeaturePluginRegistered(const FString& PluginURL, bool bCheckForRegistering /*= false*/) const
 {
 	if (const UGameFeaturePluginStateMachine* StateMachine = FindGameFeaturePluginStateMachine(PluginURL))
@@ -884,7 +913,7 @@ void UGameFeaturesSubsystem::LoadGameFeaturePlugin(const FString& PluginURL, con
 		// on whichever Role runs second between client and server.
 
 		// Refire the observer for Activated and do nothing else.
-		CallbackObservers(EObserverCallback::Activating, PluginURL, &StateMachine->GetPluginName(), StateMachine->GetGameFeatureDataForActivePlugin());
+		CallbackObservers(EObserverCallback::Activating, StateMachine->GetPluginIdentifier(), &StateMachine->GetPluginName(), StateMachine->GetGameFeatureDataForActivePlugin());
 	}
 
 	if (ShouldUpdatePluginProtocolOptions(StateMachine, ProtocolOptions))
@@ -1018,7 +1047,7 @@ void UGameFeaturesSubsystem::ChangeGameFeatureTargetState(const FString& PluginU
 		// on whichever Role runs second between client and server.
 
 		// Refire the observer for Activated and do nothing else.
-		CallbackObservers(EObserverCallback::Activating, PluginURL, &StateMachine->GetPluginName(), StateMachine->GetGameFeatureDataForActivePlugin());
+		CallbackObservers(EObserverCallback::Activating, StateMachine->GetPluginIdentifier(), &StateMachine->GetPluginName(), StateMachine->GetGameFeatureDataForActivePlugin());
 	}
 	
 	if (ShouldUpdatePluginProtocolOptions(StateMachine, ProtocolOptions))
@@ -2141,12 +2170,12 @@ void UGameFeaturesSubsystem::ListGameFeaturePlugins(const TArray<FString>& Args,
 	Ar.Logf(TEXT("Total Game Feature Plugins: %d"), PluginCount);
 }
 
-void UGameFeaturesSubsystem::CallbackObservers(EObserverCallback CallbackType, const FString& PluginURL, 
+void UGameFeaturesSubsystem::CallbackObservers(EObserverCallback CallbackType, const FGameFeaturePluginIdentifier& PluginIdentifier,
 	const FString* PluginName /*= nullptr*/, 
 	const UGameFeatureData* GameFeatureData /*= nullptr*/, 
 	FGameFeatureStateChangeContext* StateChangeContext /*= nullptr*/)
 {
-	static_assert(std::underlying_type<EObserverCallback>::type(EObserverCallback::Count) == 8, "Update UGameFeaturesSubsystem::CallbackObservers to handle added EObserverCallback");
+	static_assert(std::underlying_type<EObserverCallback>::type(EObserverCallback::Count) == 10, "Update UGameFeaturesSubsystem::CallbackObservers to handle added EObserverCallback");
 
 	// Protect against modifying the observer list during iteration
 	TArray<UObject*> LocalObservers(Observers);
@@ -2157,7 +2186,7 @@ void UGameFeaturesSubsystem::CallbackObservers(EObserverCallback CallbackType, c
 	{
 		for (UObject* Observer : LocalObservers)
 		{
-			CastChecked<IGameFeatureStateChangeObserver>(Observer)->OnGameFeatureCheckingStatus(PluginURL);
+			CastChecked<IGameFeatureStateChangeObserver>(Observer)->OnGameFeatureCheckingStatus(PluginIdentifier.GetFullPluginURL());
 		}
 		break;
 	}
@@ -2165,7 +2194,27 @@ void UGameFeaturesSubsystem::CallbackObservers(EObserverCallback CallbackType, c
 	{
 		for (UObject* Observer : LocalObservers)
 		{
-			CastChecked<IGameFeatureStateChangeObserver>(Observer)->OnGameFeatureTerminating(PluginURL);
+			CastChecked<IGameFeatureStateChangeObserver>(Observer)->OnGameFeatureTerminating(PluginIdentifier.GetFullPluginURL());
+		}
+		break;
+	}
+	case EObserverCallback::PreMounting:
+	{
+		check(PluginName);
+		FGameFeaturePreMountingContext* PreMountingContext = static_cast<FGameFeaturePreMountingContext*>(StateChangeContext);
+		for (UObject* Observer : LocalObservers)
+		{
+			CastChecked<IGameFeatureStateChangeObserver>(Observer)->OnGameFeaturePreMounting(*PluginName, PluginIdentifier, *PreMountingContext);
+		}
+		break;
+	}
+	case EObserverCallback::PostMounting:
+	{
+		check(PluginName);
+		FGameFeaturePostMountingContext* PostMountingContext = static_cast<FGameFeaturePostMountingContext*>(StateChangeContext);
+		for (UObject* Observer : LocalObservers)
+		{
+			CastChecked<IGameFeatureStateChangeObserver>(Observer)->OnGameFeaturePostMounting(*PluginName, PluginIdentifier, *PostMountingContext);
 		}
 		break;
 	}
@@ -2175,7 +2224,7 @@ void UGameFeaturesSubsystem::CallbackObservers(EObserverCallback CallbackType, c
 		check(GameFeatureData);
 		for (UObject* Observer : LocalObservers)
 		{
-			CastChecked<IGameFeatureStateChangeObserver>(Observer)->OnGameFeatureRegistering(GameFeatureData, *PluginName, PluginURL);
+			CastChecked<IGameFeatureStateChangeObserver>(Observer)->OnGameFeatureRegistering(GameFeatureData, *PluginName, PluginIdentifier.GetFullPluginURL());
 		}
 		break;
 	}
@@ -2185,7 +2234,7 @@ void UGameFeaturesSubsystem::CallbackObservers(EObserverCallback CallbackType, c
 		check(GameFeatureData);
 		for (UObject* Observer : LocalObservers)
 		{
-			CastChecked<IGameFeatureStateChangeObserver>(Observer)->OnGameFeatureUnregistering(GameFeatureData, *PluginName, PluginURL);
+			CastChecked<IGameFeatureStateChangeObserver>(Observer)->OnGameFeatureUnregistering(GameFeatureData, *PluginName, PluginIdentifier.GetFullPluginURL());
 		}
 		break;
 	}
@@ -2194,7 +2243,7 @@ void UGameFeaturesSubsystem::CallbackObservers(EObserverCallback CallbackType, c
 		check(GameFeatureData);
 		for (UObject* Observer : LocalObservers)
 		{
-			CastChecked<IGameFeatureStateChangeObserver>(Observer)->OnGameFeatureLoading(GameFeatureData, PluginURL);
+			CastChecked<IGameFeatureStateChangeObserver>(Observer)->OnGameFeatureLoading(GameFeatureData, PluginIdentifier.GetFullPluginURL());
 		}
 		break;
 	}
@@ -2203,7 +2252,7 @@ void UGameFeaturesSubsystem::CallbackObservers(EObserverCallback CallbackType, c
 		check(GameFeatureData);
 		for (UObject* Observer : LocalObservers)
 		{
-			CastChecked<IGameFeatureStateChangeObserver>(Observer)->OnGameFeatureActivating(GameFeatureData, PluginURL);
+			CastChecked<IGameFeatureStateChangeObserver>(Observer)->OnGameFeatureActivating(GameFeatureData, PluginIdentifier.GetFullPluginURL());
 		}
 		break;
 	}
@@ -2216,7 +2265,7 @@ void UGameFeaturesSubsystem::CallbackObservers(EObserverCallback CallbackType, c
 		{
 			for (UObject* Observer : LocalObservers)
 			{
-				CastChecked<IGameFeatureStateChangeObserver>(Observer)->OnGameFeatureDeactivating(GameFeatureData, *DeactivatingContext, PluginURL);
+				CastChecked<IGameFeatureStateChangeObserver>(Observer)->OnGameFeatureDeactivating(GameFeatureData, *DeactivatingContext, PluginIdentifier.GetFullPluginURL());
 			}
 		}
 		break;
@@ -2230,7 +2279,7 @@ void UGameFeaturesSubsystem::CallbackObservers(EObserverCallback CallbackType, c
 		{
 			for (UObject* Observer : LocalObservers)
 			{
-				CastChecked<IGameFeatureStateChangeObserver>(Observer)->OnGameFeaturePauseChange(PluginURL, *PluginName, *PauseChangeContext);
+				CastChecked<IGameFeatureStateChangeObserver>(Observer)->OnGameFeaturePauseChange(PluginIdentifier.GetFullPluginURL(), *PluginName, *PauseChangeContext);
 			}
 		}
 		break;
diff --git a/Engine/Plugins/Experimental/GameFeatures/Source/GameFeatures/Public/GameFeatureStateChangeObserver.h b/Engine/Plugins/Experimental/GameFeatures/Source/GameFeatures/Public/GameFeatureStateChangeObserver.h
--- a/Engine/Plugins/Experimental/GameFeatures/Source/GameFeatures/Public/GameFeatureStateChangeObserver.h
+++ b/Engine/Plugins/Experimental/GameFeatures/Source/GameFeatures/Public/GameFeatureStateChangeObserver.h
@@ -6,6 +6,9 @@
 #include "GameFeatureStateChangeObserver.generated.h"
 
 class UGameFeatureData;
+struct FGameFeaturePluginIdentifier;
+struct FGameFeaturePreMountingContext;
+struct FGameFeaturePostMountingContext;
 struct FGameFeatureDeactivatingContext;
 struct FGameFeaturePauseStateChangeContext;
 
@@ -31,18 +34,31 @@ class GAMEFEATURES_API IGameFeatureStateChangeObserver
 
 public:
 
+	// Invoked when going from the UnknownStatus state to the CheckingStatus state
 	virtual void OnGameFeatureCheckingStatus(const FString& PluginURL) {}
 
+	// Invoked prior to terminating a game feature plugin
 	virtual void OnGameFeatureTerminating(const FString& PluginURL) {}
 
+	// Invoked prior to mounting a plugin (but after its install bundles become available, if any)
+	virtual void OnGameFeaturePreMounting(const FString& PluginName, const FGameFeaturePluginIdentifier& PluginIdentifier, FGameFeaturePreMountingContext& Context) {}
+
+	// Invoked at the end of the plugin mounting phase (whether it was successfully mounted or not)
+	virtual void OnGameFeaturePostMounting(const FString& PluginName, const FGameFeaturePluginIdentifier& PluginIdentifier, FGameFeaturePostMountingContext& Context) {}
+
+	// Invoked after a game feature plugin has been registered
 	virtual void OnGameFeatureRegistering(const UGameFeatureData* GameFeatureData, const FString& PluginName, const FString& PluginURL) {}
 
+	// Invoked prior to unregistering a game feature plugin
 	virtual void OnGameFeatureUnregistering(const UGameFeatureData* GameFeatureData, const FString& PluginName, const FString& PluginURL) {}
 
+	// Invoked in the early stages of the game feature plugin loading phase
 	virtual void OnGameFeatureLoading(const UGameFeatureData* GameFeatureData, const FString& PluginURL) {}
 
+	// Invoked prior to activating a game feature plugin
 	virtual void OnGameFeatureActivating(const UGameFeatureData* GameFeatureData, const FString& PluginURL) {}
 
+	// Invoked prior to deactivating a game feature plugin
 	virtual void OnGameFeatureDeactivating(const UGameFeatureData* GameFeatureData, FGameFeatureDeactivatingContext& Context, const FString& PluginURL) {}
 
 	/** Called whenever a GameFeature State either pauses or resumes work without transitioning out of that state.
diff --git a/Engine/Plugins/Experimental/GameFeatures/Source/GameFeatures/Public/GameFeaturesSubsystem.h b/Engine/Plugins/Experimental/GameFeatures/Source/GameFeatures/Public/GameFeaturesSubsystem.h
--- a/Engine/Plugins/Experimental/GameFeatures/Source/GameFeatures/Public/GameFeaturesSubsystem.h
+++ b/Engine/Plugins/Experimental/GameFeatures/Source/GameFeatures/Public/GameFeaturesSubsystem.h
@@ -155,6 +155,41 @@ private:
 	bool bIsPaused = false;
 };
 
+
+/** Context that provides extra information prior to mounting a plugin */
+struct FGameFeaturePreMountingContext : public FGameFeatureStateChangeContext
+{
+public:
+	//@TODO: Add rules specific to pre-mounting when required
+
+private:
+
+	friend struct FGameFeaturePluginState_Mounting;
+};
+
+/** Context that allows pausing prior to transitioning out of the mounting state */
+struct FGameFeaturePostMountingContext : public FGameFeatureStateChangeContext
+{
+public:
+	// Call this if your observer has an asynchronous action to complete prior to transitioning out of the mounting state
+	// and invoke the returned delegate when you are done (on the game thread!)
+	GAMEFEATURES_API FSimpleDelegate PauseUntilComplete(FString InPauserTag);
+
+	FGameFeaturePostMountingContext(FStringView InPluginName, TFunction<void(FStringView InPauserTag)>&& InCompletionCallback)
+		: PluginName(InPluginName)
+		, CompletionCallback(MoveTemp(InCompletionCallback))
+	{}
+
+	int32 GetNumPausers() const { return NumPausers; }
+
+private:
+	FStringView PluginName;
+	TFunction<void(FStringView InPauserTag)> CompletionCallback;
+	int32 NumPausers = 0;
+
+	friend struct FGameFeaturePluginState_Mounting;
+};
+
 GAMEFEATURES_API DECLARE_LOG_CATEGORY_EXTERN(LogGameFeatures, Log, All);
 /** Notification that a game feature plugin install/register/load/unload has finished */
 DECLARE_DELEGATE_OneParam(FGameFeaturePluginChangeStateComplete, const UE::GameFeatures::FResult& /*Result*/);
@@ -226,7 +261,7 @@ struct FBuiltInGameFeaturePluginBehaviorOptions
 /** Struct used to transform a GameFeaturePlugin URL into something that can uniquely identify the GameFeaturePlugin
     without including any transient data being passed in through the URL */
 USTRUCT()
-struct FGameFeaturePluginIdentifier
+struct GAMEFEATURES_API FGameFeaturePluginIdentifier
 {
 	GENERATED_BODY()
 
@@ -411,6 +446,9 @@ public:
 	/** Determines if a plugin is in the Installed state (or beyond) */
 	bool IsGameFeaturePluginInstalled(const FString& PluginURL) const;
 
+	/** Determines if a plugin is beyond the Mounting state */
+	bool IsGameFeaturePluginMounted(const FString& PluginURL) const;
+
 	/** Determines if a plugin is in the Registered state (or beyond) */
 	bool IsGameFeaturePluginRegistered(const FString& PluginURL, bool bCheckForRegistering = false) const;
 
@@ -528,31 +566,35 @@ public:
 private:
 	TSet<FString> GetActivePluginNames() const;
 
-	void OnGameFeatureTerminating(const FString& PluginName, const FString& PluginURL);
+	void OnGameFeatureTerminating(const FString& PluginName, const FGameFeaturePluginIdentifier& PluginIdentifier);
 	friend struct FGameFeaturePluginState_Terminal;
 
-	void OnGameFeatureCheckingStatus(const FString& PluginURL);
+	void OnGameFeatureCheckingStatus(const FGameFeaturePluginIdentifier& PluginIdentifier);
 	friend struct FGameFeaturePluginState_UnknownStatus;
 
-	void OnGameFeatureStatusKnown(const FString& PluginName, const FString& PluginURL);
+	void OnGameFeatureStatusKnown(const FString& PluginName, const FGameFeaturePluginIdentifier& PluginIdentifierL);
 	friend struct FGameFeaturePluginState_CheckingStatus;
 
-	void OnGameFeatureRegistering(const UGameFeatureData* GameFeatureData, const FString& PluginName, const FString& PluginURL);
+	void OnGameFeaturePreMounting(const FString& PluginName, const FGameFeaturePluginIdentifier& PluginIdentifier, FGameFeaturePreMountingContext& Context);
+	void OnGameFeaturePostMounting(const FString& PluginName, const FGameFeaturePluginIdentifier& PluginIdentifier, FGameFeaturePostMountingContext& Context);
+	friend struct FGameFeaturePluginState_Mounting;
+
+	void OnGameFeatureRegistering(const UGameFeatureData* GameFeatureData, const FString& PluginName, const FGameFeaturePluginIdentifier& PluginIdentifier);
 	friend struct FGameFeaturePluginState_Registering;
 
-	void OnGameFeatureUnregistering(const UGameFeatureData* GameFeatureData, const FString& PluginName, const FString& PluginURL);
+	void OnGameFeatureUnregistering(const UGameFeatureData* GameFeatureData, const FString& PluginName, const FGameFeaturePluginIdentifier& PluginIdentifier);
 	friend struct FGameFeaturePluginState_Unregistering;
 
-	void OnGameFeatureActivating(const UGameFeatureData* GameFeatureData, const FString& PluginName, FGameFeatureActivatingContext& Context, const FString& PluginURL);
+	void OnGameFeatureActivating(const UGameFeatureData* GameFeatureData, const FString& PluginName, FGameFeatureActivatingContext& Context, const FGameFeaturePluginIdentifier& PluginIdentifier);
 	friend struct FGameFeaturePluginState_Activating;
 
-	void OnGameFeatureDeactivating(const UGameFeatureData* GameFeatureData, const FString& PluginName, FGameFeatureDeactivatingContext& Context, const FString& PluginURL);
+	void OnGameFeatureDeactivating(const UGameFeatureData* GameFeatureData, const FString& PluginName, FGameFeatureDeactivatingContext& Context, const FGameFeaturePluginIdentifier& PluginIdentifier);
 	friend struct FGameFeaturePluginState_Deactivating;
 
-	void OnGameFeatureLoading(const UGameFeatureData* GameFeatureData, const FString& PluginURL);
+	void OnGameFeatureLoading(const UGameFeatureData* GameFeatureData, const FGameFeaturePluginIdentifier& PluginIdentifier);
 	friend struct FGameFeaturePluginState_Loading;
 
-	void OnGameFeaturePauseChange(const FString& PluginURL, const FString& PluginName, FGameFeaturePauseStateChangeContext& Context);
+	void OnGameFeaturePauseChange(const FGameFeaturePluginIdentifier& PluginIdentifier, const FString& PluginName, FGameFeaturePauseStateChangeContext& Context);
 	friend struct FGameFeaturePluginState_Downloading;
 	friend struct FGameFeaturePluginState_Deactivating;
 
@@ -621,6 +663,8 @@ private:
 	{
 		CheckingStatus,
 		Terminating,
+		PreMounting,
+		PostMounting,
 		Registering,
 		Unregistering,
 		Loading,
@@ -630,7 +674,7 @@ private:
 		Count
 	};
 
-	void CallbackObservers(EObserverCallback CallbackType, const FString& PluginURL, 
+	void CallbackObservers(EObserverCallback CallbackType, const FGameFeaturePluginIdentifier& PluginIdentifier,
 		const FString* PluginName = nullptr, 
 		const UGameFeatureData* GameFeatureData = nullptr, 
 		FGameFeatureStateChangeContext* StateChangeContext = nullptr);
