GFS: UpdateStateMachineImmediate in OnPostMountPauserCompleted so plugin batch loading can be executed in a single frame (when no download involved)
#rnx

diff --git a/Engine/Plugins/Experimental/GameFeatures/Source/GameFeatures/Private/GameFeaturePluginStateMachine.cpp b/Engine/Plugins/Experimental/GameFeatures/Source/GameFeatures/Private/GameFeaturePluginStateMachine.cpp
--- a/Engine/Plugins/Experimental/GameFeatures/Source/GameFeatures/Private/GameFeaturePluginStateMachine.cpp
+++ b/Engine/Plugins/Experimental/GameFeatures/Source/GameFeatures/Private/GameFeaturePluginStateMachine.cpp
@@ -1735,7 +1735,7 @@ struct FGameFeaturePluginState_Mounting : public FGameFeaturePluginState
 
 		if (NumObservedPostMountPausers == NumExpectedPostMountPausers)
 		{
-			UpdateStateMachineDeferred();
+			UpdateStateMachineImmediate();
 		}
 	}
 
