Send all plugins to at minimum registerd during cooking. Installed should not prevent a plugin from being cooked.
For a plugin to be cooked it must at least be in the registered state.

diff --git a/Engine/Plugins/Experimental/GameFeatures/Source/GameFeatures/Private/GameFeaturesSubsystem.cpp b/Engine/Plugins/Experimental/GameFeatures/Source/GameFeatures/Private/GameFeaturesSubsystem.cpp
--- a/Engine/Plugins/Experimental/GameFeatures/Source/GameFeatures/Private/GameFeaturesSubsystem.cpp
+++ b/Engine/Plugins/Experimental/GameFeatures/Source/GameFeatures/Private/GameFeaturesSubsystem.cpp
@@ -1436,9 +1436,13 @@ void UGameFeaturesSubsystem::LoadBuiltInGameFeaturePlugin(const TSharedRef<IPlug
 			{
 				UGameFeaturePluginStateMachine* StateMachine = FindOrCreateGameFeaturePluginStateMachine(PluginURL, FGameFeatureProtocolOptions());
 
-				const EBuiltInAutoState InitialAutoState = (BehaviorOptions.AutoStateOverride != EBuiltInAutoState::Invalid) ? 
+				EBuiltInAutoState InitialAutoState = (BehaviorOptions.AutoStateOverride != EBuiltInAutoState::Invalid) ? 
 					BehaviorOptions.AutoStateOverride : PluginDetails.BuiltInAutoState;
-				
+				if (InitialAutoState < EBuiltInAutoState::Registered && IsRunningCookCommandlet())
+				{
+					UE_LOG(LogGameFeatures, Display, TEXT("%s will be set to Registerd for cooking"), *Plugin->GetName());
+					InitialAutoState = EBuiltInAutoState::Registered;
+				}
 				const EGameFeaturePluginState DestinationState = ConvertInitialFeatureStateToTargetState(InitialAutoState);
 
 				// If we're already at the destination or beyond, don't transition back
