Fix handling of PluginInstalledFilename when it's not relative

diff --git a/Engine/Plugins/Experimental/GameFeatures/Source/GameFeatures/Private/GameFeatureData.cpp b/Engine/Plugins/Experimental/GameFeatures/Source/GameFeatures/Private/GameFeatureData.cpp
--- a/Engine/Plugins/Experimental/GameFeatures/Source/GameFeatures/Private/GameFeatureData.cpp
+++ b/Engine/Plugins/Experimental/GameFeatures/Source/GameFeatures/Private/GameFeatureData.cpp
@@ -128,13 +128,19 @@ void UGameFeatureData::InitializeHierarchicalPluginIniFiles(const FString& Plugi
 	const FString PlatformName = FPlatformProperties::IniPlatformName();
 #endif
 
-	const FString PluginName = FPaths::GetBaseFilename(PluginInstalledFilename);
-	const FString PluginConfigDir = FPaths::GetPath(PluginInstalledFilename) / TEXT("Config/");
+	FString PluginInstalledStandardFilename = PluginInstalledFilename;
+	if (!FPaths::IsRelative(PluginInstalledStandardFilename))
+	{
+		FPaths::MakeStandardFilename(PluginInstalledStandardFilename);
+	}
+
+	const FString PluginName = FPaths::GetBaseFilename(PluginInstalledStandardFilename);
+	const FString PluginConfigDir = FPaths::GetPath(PluginInstalledStandardFilename) / TEXT("Config/");
 	const FString PlatformExtensionDir = FPaths::ProjectPlatformExtensionsDir() / (PlatformName + "/");
 	const FString EngineConfigDir = FPaths::EngineConfigDir();
 
 	// We'll look first in the game's platform extension dir for a plugin filesystem, and then default to the plugin folder
-	FString PluginPlatformConfigDir = FPaths::GetPath(PluginInstalledFilename).Replace(*FPaths::ProjectDir(), *PlatformExtensionDir) / TEXT("Config");
+	FString PluginPlatformConfigDir = FPaths::GetPath(PluginInstalledStandardFilename).Replace(*FPaths::ProjectDir(), *PlatformExtensionDir) / TEXT("Config");
 	if (!FPaths::DirectoryExists(PluginPlatformConfigDir))
 	{
 		PluginPlatformConfigDir = FPaths::Combine(PluginConfigDir, PlatformName);
