Fix performance regression where cooking incorrectly uses the running platform's device profile

#rnx

diff --git a/Engine/Plugins/Experimental/GameFeatures/Source/GameFeatures/Private/GameFeatureData.cpp b/Engine/Plugins/Experimental/GameFeatures/Source/GameFeatures/Private/GameFeatureData.cpp
--- a/Engine/Plugins/Experimental/GameFeatures/Source/GameFeatures/Private/GameFeatureData.cpp
+++ b/Engine/Plugins/Experimental/GameFeatures/Source/GameFeatures/Private/GameFeatureData.cpp
@@ -375,7 +375,22 @@ void UGameFeatureData::InitializeHierarchicalPluginIniFiles(const FString& Plugi
 			// Need to get the in-memory config filename, the on disk one is likely not up to date
 			FString IniFile = GConfig->GetConfigFilename(*Ini.Name);
 
-			if (FConfigFile* ExistingConfig = GConfig->FindConfigFile(IniFile))
+			// Ensure we push new device profile config to the appropriate config branch - GConfig could be Windows while we're previewing a console
+			FConfigFile* ExistingConfig = nullptr;
+#if ALLOW_OTHER_PLATFORM_CONFIG
+			if (Ini.bCreateDeviceProfiles && Ini.bUsePlatformDir)
+			{
+				FConfigCacheIni* PlatformConfigSystem = FConfigCacheIni::ForPlatform(*PlatformName);
+				ExistingConfig = PlatformConfigSystem->FindConfigFile(GDeviceProfilesIni);
+			}
+#endif
+
+			if (ExistingConfig == nullptr)
+			{
+				ExistingConfig = GConfig->FindConfigFile(IniFile);
+			}
+
+			if (ExistingConfig)
 			{
 				if (Ini.bCreateDeviceProfiles)
 				{
diff --git a/Engine/Source/Runtime/Engine/Private/DeviceProfiles/DeviceProfileManager.cpp b/Engine/Source/Runtime/Engine/Private/DeviceProfiles/DeviceProfileManager.cpp
--- a/Engine/Source/Runtime/Engine/Private/DeviceProfiles/DeviceProfileManager.cpp
+++ b/Engine/Source/Runtime/Engine/Private/DeviceProfiles/DeviceProfileManager.cpp
@@ -181,13 +181,8 @@ TMap<FName, FString> UDeviceProfileManager::GatherDeviceProfileCVars(const FStri
 			return DeviceProfileCVars;
 		}
 
-		// use the DP's platform's configs, NOT the running platform, unless the profile is not there
-		FConfigCacheIni* PlatformConfigSystem = FConfigCacheIni::ForPlatform(*Profile->DeviceType);
-		FConfigFile* PlatformConfigFile = PlatformConfigSystem->FindConfigFile(GDeviceProfilesIni);
-		if (PlatformConfigFile && PlatformConfigFile->Contains(DeviceProfileName))
-		{
-			ConfigSystem = PlatformConfigSystem;
-		}
+		// use the DP's platform's configs, NOT the running platform
+		ConfigSystem = FConfigCacheIni::ForPlatform(*Profile->DeviceType);
 
 		MemBucket = Profile->GetPreviewMemorySizeBucket();
 		if(GatherMode == EDeviceProfileMode::DPM_CacheValues)
