Sort the plugins before loading all built in plugins so the leaf plugins are loaded first. This gives a more accurate count to how long a plugin takes to load and removes the need to perform extra work subtracting to calculate the exclusive plugin load time. We are now always using the exclusive plugin load time by sorting.

diff --git a/Engine/Plugins/Experimental/GameFeatures/Source/GameFeatures/Private/GameFeaturesSubsystem.cpp b/Engine/Plugins/Experimental/GameFeatures/Source/GameFeatures/Private/GameFeaturesSubsystem.cpp
--- a/Engine/Plugins/Experimental/GameFeatures/Source/GameFeatures/Private/GameFeaturesSubsystem.cpp
+++ b/Engine/Plugins/Experimental/GameFeatures/Source/GameFeatures/Private/GameFeaturesSubsystem.cpp
@@ -1,6 +1,7 @@
 // Copyright Epic Games, Inc. All Rights Reserved.
 
 #include "GameFeaturesSubsystem.h"
+#include "Algo/TopologicalSort.h"
 #include "AssetRegistry/AssetData.h"
 #include "AssetRegistry/IAssetRegistry.h"
 #include "GameFeaturesSubsystemSettings.h"
@@ -1580,6 +1581,28 @@ void UGameFeaturesSubsystem::LoadBuiltInGameFeaturePlugins(FBuiltInPluginAdditio
 	const int32 NumPluginsToLoad = EnabledPlugins.Num();
 	UE_LOG(LogGameFeatures, Log, TEXT("Loading %i builtins"), NumPluginsToLoad);
 
+	// Sort the plugins so we can more accurately track how long it takes to load rather than have inconsistent dependency timings.
+	auto GetPluginDependencies =
+		[](TSharedRef<IPlugin> CurrentPlugin)
+	{
+		IPluginManager& PluginManager = IPluginManager::Get();
+
+		TArray<TSharedRef<IPlugin>> Dependencies;
+		const FPluginDescriptor& Desc = CurrentPlugin->GetDescriptor();
+		for (const FPluginReferenceDescriptor& Dependency : Desc.Plugins)
+		{
+			if (Dependency.bEnabled)
+			{
+				if (TSharedPtr<IPlugin> FoundPlugin = PluginManager.FindPlugin(Dependency.Name))
+				{
+					Dependencies.Add(FoundPlugin.ToSharedRef());
+				}
+			}
+		}
+		return Dependencies;
+	};
+	Algo::TopologicalSort(EnabledPlugins, GetPluginDependencies);
+
 	for (const TSharedRef<IPlugin>& Plugin : EnabledPlugins)
 	{
 		FBuiltInPluginLoadTimeTrackerScope TrackerScope(PluginLoadTimeTracker, Plugin);
@@ -2441,4 +2464,3 @@ EGameFeaturePluginState UGameFeaturesSubsystem::ConvertInitialFeatureStateToTarg
 	}
 	return InitialState;
 }
-
