Perf improvement

diff --git a/Engine/Plugins/Experimental/GameFeatures/Source/GameFeatures/Private/GameFeaturesSubsystem.cpp b/Engine/Plugins/Experimental/GameFeatures/Source/GameFeatures/Private/GameFeaturesSubsystem.cpp
--- a/Engine/Plugins/Experimental/GameFeatures/Source/GameFeatures/Private/GameFeaturesSubsystem.cpp
+++ b/Engine/Plugins/Experimental/GameFeatures/Source/GameFeatures/Private/GameFeaturesSubsystem.cpp
@@ -1582,12 +1582,13 @@ void UGameFeaturesSubsystem::LoadBuiltInGameFeaturePlugins(FBuiltInPluginAdditio
 	UE_LOG(LogGameFeatures, Log, TEXT("Loading %i builtins"), NumPluginsToLoad);
 
 	// Sort the plugins so we can more accurately track how long it takes to load rather than have inconsistent dependency timings.
+	TArray<TSharedRef<IPlugin>> Dependencies;
 	auto GetPluginDependencies =
-		[](TSharedRef<IPlugin> CurrentPlugin)
+		[&Dependencies](TSharedRef<IPlugin> CurrentPlugin)
 	{
 		IPluginManager& PluginManager = IPluginManager::Get();
-
-		TArray<TSharedRef<IPlugin>> Dependencies;
+		Dependencies.Reset();
+		
 		const FPluginDescriptor& Desc = CurrentPlugin->GetDescriptor();
 		for (const FPluginReferenceDescriptor& Dependency : Desc.Plugins)
 		{
