Submitting on behalf of BobT:

Change AddToPlaylists to by default include the plugin the GFA is in without putting in the name of the plugin

diff --git a/Engine/Plugins/Experimental/GameFeatures/Source/GameFeatures/Private/GameFeatureAction.cpp b/Engine/Plugins/Experimental/GameFeatures/Source/GameFeatures/Private/GameFeatureAction.cpp
--- a/Engine/Plugins/Experimental/GameFeatures/Source/GameFeatures/Private/GameFeatureAction.cpp
+++ b/Engine/Plugins/Experimental/GameFeatures/Source/GameFeatures/Private/GameFeatureAction.cpp
@@ -1,9 +1,23 @@
 // Copyright Epic Games, Inc. All Rights Reserved.
 
 #include "GameFeatureAction.h"
+#include "GameFeatureData.h"
 
 #include UE_INLINE_GENERATED_CPP_BY_NAME(GameFeatureAction)
 
+UGameFeatureData* UGameFeatureAction::GetGameFeatureData() const
+{
+	for (UObject* Obj = GetOuter(); Obj; Obj = Obj->GetOuter())
+	{
+		if (UGameFeatureData* GFD = Cast<UGameFeatureData>(Obj))
+		{
+			return GFD;
+		}
+	}
+
+	return nullptr;
+}
+
 void UGameFeatureAction::OnGameFeatureActivating(FGameFeatureActivatingContext& Context)
 {
 	// Call older style if not overridden
diff --git a/Engine/Plugins/Experimental/GameFeatures/Source/GameFeatures/Public/GameFeatureAction.h b/Engine/Plugins/Experimental/GameFeatures/Source/GameFeatures/Public/GameFeatureAction.h
--- a/Engine/Plugins/Experimental/GameFeatures/Source/GameFeatures/Public/GameFeatureAction.h
+++ b/Engine/Plugins/Experimental/GameFeatures/Source/GameFeatures/Public/GameFeatureAction.h
@@ -4,6 +4,7 @@
 
 #include "GameFeatureAction.generated.h"
 
+class UGameFeatureData;
 struct FGameFeatureActivatingContext;
 struct FGameFeatureDeactivatingContext;
 struct FAssetBundleData;
@@ -15,6 +16,7 @@ class GAMEFEATURES_API UGameFeatureAction : public UObject
 	GENERATED_BODY()
 
 public:
+	virtual UGameFeatureData* GetGameFeatureData() const;
 
 	/** Called when the object owning the action is registered for possible activation, this is called even if a feature never activates */
 	virtual void OnGameFeatureRegistering() {}
diff --git a/Engine/Plugins/Experimental/GameFeatures/Source/GameFeatures/Public/GameFeatureAction_AddWPContent.h b/Engine/Plugins/Experimental/GameFeatures/Source/GameFeatures/Public/GameFeatureAction_AddWPContent.h
--- a/Engine/Plugins/Experimental/GameFeatures/Source/GameFeatures/Public/GameFeatureAction_AddWPContent.h
+++ b/Engine/Plugins/Experimental/GameFeatures/Source/GameFeatures/Public/GameFeatureAction_AddWPContent.h
@@ -27,8 +27,6 @@ public:
 	virtual void OnGameFeatureDeactivating(FGameFeatureDeactivatingContext& Context) override;
 	//~ End UGameFeatureAction interface
 
-	UGameFeatureData* GetGameFeatureData() const;
-
 	const UContentBundleDescriptor* GetContentBundleDescriptor() const { return ContentBundleDescriptor; }
 
 private:
