Adding a "realtime mode" on the GFS to distribute plugin state updates over several frames
- cvar activated for now, a clean api could be added later once the feature has been proven effective and also some state transitions don't support it atm
#rnx

diff --git a/Engine/Plugins/Experimental/GameFeatures/Source/GameFeatures/Private/GameFeaturePluginStateMachine.cpp b/Engine/Plugins/Experimental/GameFeatures/Source/GameFeatures/Private/GameFeaturePluginStateMachine.cpp
--- a/Engine/Plugins/Experimental/GameFeatures/Source/GameFeatures/Private/GameFeaturePluginStateMachine.cpp
+++ b/Engine/Plugins/Experimental/GameFeatures/Source/GameFeatures/Private/GameFeaturePluginStateMachine.cpp
@@ -32,6 +32,7 @@
 #include "UObject/UObjectIterator.h"
 #include "UObject/NoExportTypes.h"
 #include "Misc/PathViews.h"
+#include "Containers/Queue.h"
 
 #include UE_INLINE_GENERATED_CPP_BY_NAME(GameFeaturePluginStateMachine)
 
@@ -214,6 +215,102 @@ namespace UE::GameFeatures
 		}
 	};
 
+	static bool bRealtimeMode = false;
+
+	class FRealtimeMode : public TSharedFromThis<FRealtimeMode>
+	{
+	public:
+		~FRealtimeMode()
+		{
+			if (TickHandle.IsValid())
+			{
+				FTSTicker::GetCoreTicker().RemoveTicker(MoveTemp(TickHandle));
+			}
+
+			FGameFeaturePluginRequestUpdateStateMachine UpdateRequest;
+			while (UpdateRequests.Dequeue(UpdateRequest))
+			{
+				UpdateRequest.ExecuteIfBound();
+			}
+		}
+
+		void AddUpdateRequest(FGameFeaturePluginRequestUpdateStateMachine UpdateRequest)
+		{
+			UpdateRequests.Enqueue(MoveTemp(UpdateRequest));
+			EnableTick();
+		}
+
+	private:
+		void EnableTick()
+		{
+			if (!TickHandle.IsValid())
+			{
+				TickHandle = FTSTicker::GetCoreTicker().AddTicker(FTickerDelegate::CreateSP(this, &FRealtimeMode::Tick));
+			}
+		}
+
+		bool Tick(float DeltaTime)
+		{
+			// Self-reference so we don't get destroyed during Tick
+			TSharedRef<FRealtimeMode> SelfRef = AsShared();
+
+			{
+				constexpr double MaxFrameTime = 0.033; // 30fps
+				constexpr double AllottedTime = MaxFrameTime / 2;
+				const double StartTime = FPlatformTime::Seconds();
+
+				FGameFeaturePluginRequestUpdateStateMachine UpdateRequest;
+				while (UpdateRequests.Dequeue(UpdateRequest))
+				{
+					UpdateRequest.ExecuteIfBound();
+
+					const double ElapsedTime = FPlatformTime::Seconds() - StartTime;
+					if ((ElapsedTime > AllottedTime) || ((DeltaTime + ElapsedTime) > MaxFrameTime))
+					{
+						break;
+					}
+				}
+			}
+
+			if (UpdateRequests.IsEmpty())
+			{
+				TickHandle.Reset();
+				return false;
+			}
+			else
+			{
+				return true;
+			}
+		}
+
+	private:
+		TQueue<FGameFeaturePluginRequestUpdateStateMachine> UpdateRequests;
+		FTSTicker::FDelegateHandle TickHandle;
+	};
+
+	static TSharedPtr<FRealtimeMode> RealtimeMode;
+
+	static FAutoConsoleVariableRef CVarRealtimeMode(TEXT("GameFeaturePlugin.RealtimeMode"),
+		bRealtimeMode,
+		TEXT("Sets whether GFS realtime mode is enabled; which distributes plugin state updates over several frames"),
+		FConsoleVariableDelegate::CreateLambda(
+			[](IConsoleVariable* Var)
+			{
+				if (Var->GetBool())
+				{
+					if (!RealtimeMode)
+					{
+						RealtimeMode = MakeShared<FRealtimeMode>();
+					}
+				}
+				else
+				{
+					TSharedPtr<FRealtimeMode> Rm = MoveTemp(RealtimeMode);
+					Rm.Reset();
+				}
+			}),
+			ECVF_ReadOnly);
+
 #if WITH_EDITOR
 	TMap<FString, FGameFeaturePluginRequestUpdateStateMachine> PluginsToUnloadAssets;
 	FTSTicker::FDelegateHandle UnloadPluginAssetsHandle;
@@ -1700,6 +1797,7 @@ struct FGameFeaturePluginState_Mounting : public FGameFeaturePluginState
 	int32 NumObservedPostMountPausers = 0;
 	int32 NumExpectedPostMountPausers = 0;
 	TArray<FName> PendingBundles;
+	bool bCheckedRealtimeMode = false;
 	UE::GameFeatures::FResult Result;
 
 	void OnInstallBundleCompleted(FInstallBundleRequestResultInfo BundleResult)
@@ -1765,6 +1863,7 @@ struct FGameFeaturePluginState_Mounting : public FGameFeaturePluginState
 	{
 		NumObservedPostMountPausers = 0;
 		NumExpectedPostMountPausers = 0;
+		bCheckedRealtimeMode = false;
 		PendingBundles.Empty();
 		Result = MakeValue();
 
@@ -1817,8 +1916,6 @@ struct FGameFeaturePluginState_Mounting : public FGameFeaturePluginState
 
 	virtual void UpdateState(FGameFeaturePluginStateStatus& StateStatus) override
 	{
-		TRACE_CPUPROFILER_EVENT_SCOPE(GFP_Mounting);
-
 		// Check if waiting for install bundles
 		if (PendingBundles.Num() > 0)
 		{
@@ -1832,16 +1929,23 @@ struct FGameFeaturePluginState_Mounting : public FGameFeaturePluginState
 			if (NumExpectedPostMountPausers == NumObservedPostMountPausers)
 			{
 				NumExpectedPostMountPausers = INDEX_NONE;
-
-				// We previously sent an OnGameFeaturePauseChange delegate we need to send that work is now unpaused
-				FGameFeaturePauseStateChangeContext UnPauseContext(LexToString(EGameFeaturePluginState::Mounting), TEXT(""), /*bIsPausedIn=*/false);
-				UGameFeaturesSubsystem::Get().OnGameFeaturePauseChange(StateProperties.PluginIdentifier, StateProperties.PluginName, UnPauseContext);
-
 				TransitionOut(StateStatus);
 			}
 			return;
 		}
 
+		if (!bCheckedRealtimeMode)
+		{
+			bCheckedRealtimeMode = true;
+			if (UE::GameFeatures::RealtimeMode)
+			{
+				UE::GameFeatures::RealtimeMode->AddUpdateRequest(StateProperties.OnRequestUpdateStateMachine);
+				return;
+			}
+		}
+
+		TRACE_CPUPROFILER_EVENT_SCOPE(GFP_Mounting);
+
 		// Don't mount the plugin if there was an error during BeginState or bundles install
 		if (!Result.HasError())
 		{
@@ -1921,13 +2025,7 @@ struct FGameFeaturePluginState_Mounting : public FGameFeaturePluginState
 			NumExpectedPostMountPausers = Context.NumPausers;
 
 			// Check if we got post-mount paused
-			if (NumExpectedPostMountPausers > 0)
-			{
-				// Since we are pausing work during this mounting, also notify the OnGameFeaturePauseChange delegate
-				FGameFeaturePauseStateChangeContext PauseContext(LexToString(EGameFeaturePluginState::Mounting), TEXT("PendingPostMountCallbacks"), true);
-				UGameFeaturesSubsystem::Get().OnGameFeaturePauseChange(StateProperties.PluginIdentifier, StateProperties.PluginName, PauseContext);
-			}
-			else
+			if (NumExpectedPostMountPausers <= 0)
 			{
 				TransitionOut(StateStatus);
 			}
@@ -2062,8 +2160,25 @@ struct FGameFeaturePluginState_Registering : public FGameFeaturePluginState
 {
 	FGameFeaturePluginState_Registering(FGameFeaturePluginStateMachineProperties& InStateProperties) : FGameFeaturePluginState(InStateProperties) {}
 
+	bool bCheckedRealtimeMode = false;
+
+	virtual void BeginState() override
+	{
+		bCheckedRealtimeMode = false;
+	}
+
 	virtual void UpdateState(FGameFeaturePluginStateStatus& StateStatus) override
 	{
+		if (!bCheckedRealtimeMode)
+		{
+			bCheckedRealtimeMode = true;
+			if (UE::GameFeatures::RealtimeMode)
+			{
+				UE::GameFeatures::RealtimeMode->AddUpdateRequest(StateProperties.OnRequestUpdateStateMachine);
+				return;
+			}
+		}
+
 		TRACE_CPUPROFILER_EVENT_SCOPE(GFP_Registering);
 		const FString PluginFolder = FPaths::GetPath(StateProperties.PluginInstalledFilename);
 		UGameplayTagsManager::Get().AddTagIniSearchPath(PluginFolder / TEXT("Config") / TEXT("Tags"));
