GFS: re-add realtime mode on WaitingForDependencies state that was lost when merging 27558215
#rnx

diff --git a/Engine/Plugins/Experimental/GameFeatures/Source/GameFeatures/Private/GameFeaturePluginStateMachine.cpp b/Engine/Plugins/Experimental/GameFeatures/Source/GameFeatures/Private/GameFeaturePluginStateMachine.cpp
--- a/Engine/Plugins/Experimental/GameFeatures/Source/GameFeatures/Private/GameFeaturePluginStateMachine.cpp
+++ b/Engine/Plugins/Experimental/GameFeatures/Source/GameFeatures/Private/GameFeaturePluginStateMachine.cpp
@@ -582,7 +582,6 @@ struct FTransitionDependenciesGameFeaturePluginState : public FGameFeaturePlugin
 {
 	FTransitionDependenciesGameFeaturePluginState(FGameFeaturePluginStateMachineProperties& InStateProperties)
 		: FGameFeaturePluginState(InStateProperties)
-		, bRequestedDependencies(false)
 	{}
 
 	virtual ~FTransitionDependenciesGameFeaturePluginState()
@@ -593,6 +592,7 @@ struct FTransitionDependenciesGameFeaturePluginState : public FGameFeaturePlugin
 	virtual void BeginState() override
 	{
 		ClearDependencies();
+		bCheckedRealtimeMode = false;
 	}
 
 	virtual void EndState() override
@@ -602,6 +602,16 @@ struct FTransitionDependenciesGameFeaturePluginState : public FGameFeaturePlugin
 
 	virtual void UpdateState(FGameFeaturePluginStateStatus& StateStatus) override
 	{
+		if (!bCheckedRealtimeMode)
+		{
+			bCheckedRealtimeMode = true;
+			if (UE::GameFeatures::RealtimeMode)
+			{
+				UE::GameFeatures::RealtimeMode->AddUpdateRequest(StateProperties.OnRequestUpdateStateMachine);
+				return;
+			}
+		}
+
 		TRACE_CPUPROFILER_EVENT_SCOPE(GFP_TransitionDependencies);
 		checkf(!StateProperties.PluginInstalledFilename.IsEmpty(), TEXT("PluginInstalledFilename must be set by the loading dependencies phase. PluginURL: %s"), *StateProperties.PluginIdentifier.GetFullPluginURL());
 		checkf(FPaths::GetExtension(StateProperties.PluginInstalledFilename) == TEXT("uplugin"), TEXT("PluginInstalledFilename must have a uplugin extension. PluginURL: %s"), *StateProperties.PluginIdentifier.GetFullPluginURL());
@@ -754,6 +764,7 @@ struct FTransitionDependenciesGameFeaturePluginState : public FGameFeaturePlugin
 	using FDepResultPair = TPair<TWeakObjectPtr<UGameFeaturePluginStateMachine>, UE::GameFeatures::FResult>;
 	TArray<FDepResultPair> RemainingDependencies;
 	bool bRequestedDependencies = false;
+	bool bCheckedRealtimeMode = false;
 };
 
 struct FGameFeaturePluginState_Uninitialized : public FGameFeaturePluginState
