Fixes development cooked build not applying Device Profile fragments to the right config file

diff --git a/Engine/Plugins/Experimental/GameFeatures/Source/GameFeatures/Private/GameFeatureData.cpp b/Engine/Plugins/Experimental/GameFeatures/Source/GameFeatures/Private/GameFeatureData.cpp
--- a/Engine/Plugins/Experimental/GameFeatures/Source/GameFeatures/Private/GameFeatureData.cpp
+++ b/Engine/Plugins/Experimental/GameFeatures/Source/GameFeatures/Private/GameFeatureData.cpp
@@ -145,6 +145,7 @@ void UGameFeatureData::InitializeHierarchicalPluginIniFiles(const FString& Plugi
 	{
 		PluginPlatformConfigDir = FPaths::Combine(PluginConfigDir, PlatformName);
 	}
+	PluginPlatformConfigDir += TEXT("/");
 
 	const bool bIsBaseIniName = false;
 	const bool bForceReloadFromDisk = false;
@@ -378,7 +379,7 @@ void UGameFeatureData::InitializeHierarchicalPluginIniFiles(const FString& Plugi
 			// Ensure we push new device profile config to the appropriate config branch - GConfig could be Windows while we're previewing a console
 			FConfigFile* ExistingConfig = nullptr;
 #if ALLOW_OTHER_PLATFORM_CONFIG
-			if (Ini.bCreateDeviceProfiles && Ini.bUsePlatformDir)
+			if (Ini.bCreateDeviceProfiles && Ini.bUsePlatformDir && !FPlatformProperties::RequiresCookedData())
 			{
 				FConfigCacheIni* PlatformConfigSystem = FConfigCacheIni::ForPlatform(*PlatformName);
 				ExistingConfig = PlatformConfigSystem->FindConfigFile(GDeviceProfilesIni);
