Silenced an error that can happen when loading GameFeatures in UE without a project set

#rnx

diff --git a/Engine/Plugins/Experimental/GameFeatures/Source/GameFeatures/Private/GameFeaturesSubsystem.cpp b/Engine/Plugins/Experimental/GameFeatures/Source/GameFeatures/Private/GameFeaturesSubsystem.cpp
--- a/Engine/Plugins/Experimental/GameFeatures/Source/GameFeatures/Private/GameFeaturesSubsystem.cpp
+++ b/Engine/Plugins/Experimental/GameFeatures/Source/GameFeatures/Private/GameFeaturesSubsystem.cpp
@@ -11,6 +11,7 @@
 #include "GameFeatureStateChangeObserver.h"
 #include "GameplayTagsManager.h"
 #include "Interfaces/IPluginManager.h"
+#include "Misc/App.h"
 #include "Misc/PackageName.h"
 #include "Misc/Paths.h"
 #include "Misc/PathViews.h"
@@ -480,7 +481,8 @@ void UGameFeaturesSubsystem::OnAssetManagerCreated()
 	FPrimaryAssetRules GameDataRules = UAssetManager::Get().GetPrimaryAssetRules(DummyGameFeatureDataAssetId);
 	if (GameDataRules.IsDefault())
 	{
-		UE_LOG(LogGameFeatures, Error, TEXT("Asset manager settings do not include a rule for assets of type %s, which is required for game feature plugins to function"), *UGameFeatureData::StaticClass()->GetName());
+		const bool bHasProject = FApp::HasProjectName(); // Only error when we have a UE project loaded, as otherwise there won't be any GFPs to load
+		UE_CLOG(bHasProject, LogGameFeatures, Error, TEXT("Asset manager settings do not include a rule for assets of type %s, which is required for game feature plugins to function"), *UGameFeatureData::StaticClass()->GetName());
 	}
 
 	// Create the game-specific policy
