Add game feature plugin error for full cache
Improve game feature plugin errors
#rnx

diff --git a/Engine/Plugins/Experimental/GameFeatures/Source/GameFeatures/Private/GameFeaturePluginOperationResult.cpp b/Engine/Plugins/Experimental/GameFeatures/Source/GameFeatures/Private/GameFeaturePluginOperationResult.cpp
--- a/Engine/Plugins/Experimental/GameFeatures/Source/GameFeatures/Private/GameFeaturePluginOperationResult.cpp
+++ b/Engine/Plugins/Experimental/GameFeatures/Source/GameFeatures/Private/GameFeaturePluginOperationResult.cpp
@@ -19,6 +19,19 @@ namespace UE::GameFeatures
 
 	FString ToString(const FResult& Result)
 	{
-		return Result.HasValue() ? FString(TEXT("Success")) : (FString(TEXT("Failure, ErrorCode=")) + Result.GetError() + FString(TEXT(", OptionalErrorText=")) + Result.OptionalErrorText.ToString());
+		TStringBuilder<512> Out;
+		if (Result.HasValue())
+		{
+			Out << TEXT("Success");
+		}
+		else
+		{
+			Out << TEXT("ErrorCode=") << Result.GetError();
+			if (!Result.OptionalErrorText.IsEmpty())
+			{
+				Out << TEXT(", ErrorText=") << Result.OptionalErrorText.ToString();
+			}
+		}
+		return Out.ToString();
 	}
 }	// namespace UE::GameFeatures
\ No newline at end of file
diff --git a/Engine/Plugins/Experimental/GameFeatures/Source/GameFeatures/Private/GameFeaturePluginStateMachine.cpp b/Engine/Plugins/Experimental/GameFeatures/Source/GameFeatures/Private/GameFeaturePluginStateMachine.cpp
--- a/Engine/Plugins/Experimental/GameFeatures/Source/GameFeatures/Private/GameFeaturePluginStateMachine.cpp
+++ b/Engine/Plugins/Experimental/GameFeatures/Source/GameFeatures/Private/GameFeaturePluginStateMachine.cpp
@@ -137,6 +137,7 @@ namespace UE::GameFeatures
 		const FText Generic_MountError = NSLOCTEXT("GameFeatures", "ErrorCodes.MountGenericError", "An error has occurred loading data for this game feature. Please try again later.");
 
 		const FText BundleResult_NeedsUpdate = NSLOCTEXT("GameFeatures", "ErrorCodes.BundleResult.NeedsUpdate", "An application update is required to install this game feature. Please restart the application after downloading any required updates.");
+		const FText BundleResult_NeedsCacheSpace = NSLOCTEXT("GameFeatures", "ErrorCodes.BundleResult.NeedsCacheSpace", "Unable to allocate enough space in the cache to install this game feature.");
 		const FText BundleResult_NeedsDiskSpace = NSLOCTEXT("GameFeatures", "ErrorCodes.BundleResult.NeedsDiskSpace", "You do not have enough disk space to install this game feature. Please try again after clearing up disk space.");
 		const FText BundleResult_DownloadCancelled = NSLOCTEXT("GameFeatures", "ErrorCodes.BundleResult.DownloadCancelled", "This game feature download was canceled.");
 
@@ -157,12 +158,17 @@ namespace UE::GameFeatures
 
 				//These are generally unrecoverable and mean something is seriously wrong with the data for this build
 				case EInstallBundleResult::InitializationError:
-				case EInstallBundleResult::FailedCacheReserve:
 				{
 					return Generic_FatalError;
 				}
 
-				//All of these are indicative of not having enough space to install the required files
+				//Not enough space in cache to install the files
+				case EInstallBundleResult::FailedCacheReserve:
+				{
+					return BundleResult_NeedsCacheSpace;
+				}
+
+				//All of these are indicative of not having enough disk space to install the required files
 				case EInstallBundleResult::InstallerOutOfDiskSpaceError:
 				case EInstallBundleResult::ManifestArchiveError:
 				{
