Allow game feature config inis to set config values on blueprint classes.

diff --git a/Engine/Plugins/Experimental/GameFeatures/Source/GameFeatures/Private/GameFeatureData.cpp b/Engine/Plugins/Experimental/GameFeatures/Source/GameFeatures/Private/GameFeatureData.cpp
--- a/Engine/Plugins/Experimental/GameFeatures/Source/GameFeatures/Private/GameFeatureData.cpp
+++ b/Engine/Plugins/Experimental/GameFeatures/Source/GameFeatures/Private/GameFeatureData.cpp
@@ -471,7 +471,7 @@ void UGameFeatureData::ReloadConfigs(FConfigFile& PluginConfig) const
 			// Find the affected class and push updates to all instances of it, including children
 			// @note:	Intentionally not using the propagation flags inherent in ReloadConfig to handle this, as it utilizes a naive complete object iterator
 			//			and tanks performance pretty badly
-			UClass* ObjClass = FindFirstObject<UClass>(*SectionName, EFindFirstObjectOptions::ExactClass | EFindFirstObjectOptions::EnsureIfAmbiguous | EFindFirstObjectOptions::NativeFirst);
+			UClass* ObjClass = FindFirstObject<UClass>(*SectionName, EFindFirstObjectOptions::EnsureIfAmbiguous | EFindFirstObjectOptions::NativeFirst);
 			if (ObjClass)
 			{
 				TArray<UObject*> FoundObjects;
