GFP state machine can load assets async
- For now, controlled with a CVar
- default is off

diff --git a/Engine/Plugins/Experimental/GameFeatures/Source/GameFeatures/Private/GameFeaturePluginStateMachine.cpp b/Engine/Plugins/Experimental/GameFeatures/Source/GameFeatures/Private/GameFeaturePluginStateMachine.cpp
--- a/Engine/Plugins/Experimental/GameFeatures/Source/GameFeatures/Private/GameFeaturePluginStateMachine.cpp
+++ b/Engine/Plugins/Experimental/GameFeatures/Source/GameFeatures/Private/GameFeaturePluginStateMachine.cpp
@@ -52,13 +52,15 @@ namespace UE::GameFeatures
 		ShouldLogMountedFiles,
 		TEXT("Should the newly mounted files be logged."));
 
-
 	static FString GVerifyPluginSkipList;
 	static FAutoConsoleVariableRef CVarVerifyPluginSkipList(TEXT("PluginManager.VerifyUnload.SkipList"),
 		GVerifyPluginSkipList,
 		TEXT("Comma-separated list of names of plugins for which to skip verification."),
-		ECVF_Default
-		);
+		ECVF_Default);
+
+	static TAutoConsoleVariable<bool> CVarAsyncLoad(TEXT("GameFeaturePlugin.AsyncLoad"),
+		false,
+		TEXT("Enable to use aysnc loading"));
 
 	#define GAME_FEATURE_PLUGIN_STATE_TO_STRING(inEnum, inText) case EGameFeaturePluginState::inEnum: return TEXT(#inEnum);
 	FString ToString(EGameFeaturePluginState InType)
@@ -577,6 +579,11 @@ bool FGameFeaturePluginState::AllowIniLoading() const
 	}
 }
 
+bool FGameFeaturePluginState::UseAsyncLoading() const
+{
+	return UE::GameFeatures::CVarAsyncLoad.GetValueOnGameThread();
+}
+
 /*
 =========================================================
   States
@@ -1895,7 +1902,7 @@ struct FGameFeaturePluginState_Mounting : public FGameFeaturePluginState
 		const TArray<FName>& InstallBundles = MetaData.InstallBundles;
 
 		const FInstallBundlePluginProtocolOptions& Options = StateProperties.ProtocolOptions.GetSubtype<FInstallBundlePluginProtocolOptions>();
-		const EInstallBundleRequestFlags InstallFlags = Options.InstallBundleFlags;
+		const EInstallBundleRequestFlags InstallFlags = Options.InstallBundleFlags /*| EInstallBundleRequestFlags::AsyncMount*/;
 
 		// Make bundle manager use verbose log level for most logs.
 		// We are already done with downloading, so we don't care about logging too much here unless mounting fails.
@@ -2006,6 +2013,7 @@ struct FGameFeaturePluginState_Mounting : public FGameFeaturePluginState
 					IPluginManager::Get().MountExplicitlyLoadedPlugin(StateProperties.PluginName);
 				}
 
+				// @TODO: Load AR Data async?
 				// After the new plugin is mounted add the asset registry for that plugin.
 				if (StateProperties.GetPluginProtocol() == EGameFeaturePluginProtocol::InstallBundle)
 				{
@@ -2175,32 +2183,100 @@ struct FGameFeaturePluginState_Unregistering : public FGameFeaturePluginState
 
 struct FGameFeaturePluginState_Registering : public FGameFeaturePluginState
 {
-	FGameFeaturePluginState_Registering(FGameFeaturePluginStateMachineProperties& InStateProperties) : FGameFeaturePluginState(InStateProperties) {}
+	enum class ELoadGFDState : uint8
+	{
+		Pending = 0,
+		Success,
+		Cancelled,
+		Failed
+	};
 
+	TSharedPtr<FStreamableHandle> GameFeatureDataHandle;
+	TArray<FString, TInlineAllocator<2>> GameFeatureDataSearchPaths;
+	ELoadGFDState LoadGFDState = ELoadGFDState::Pending;
 	bool bCheckedRealtimeMode = false;
 
-	virtual void BeginState() override
+	FGameFeaturePluginState_Registering(FGameFeaturePluginStateMachineProperties& InStateProperties) : FGameFeaturePluginState(InStateProperties) {}
+	
+	void TryAsyncLoadGameFeatureData(int32 Attempt = 0)
 	{
-		bCheckedRealtimeMode = false;
-	}
+		if (!GameFeatureDataSearchPaths.IsValidIndex(Attempt))
+		{
+			LoadGFDState = ELoadGFDState::Failed;
+			UpdateStateMachineDeferred();
+			return;
+		}
 
-	virtual void UpdateState(FGameFeaturePluginStateStatus& StateStatus) override
-	{
-		if (!bCheckedRealtimeMode)
+		check(LoadGFDState == ELoadGFDState::Pending);
+
+		bool bIsLoading = false;
+
+		const FString& Path = GameFeatureDataSearchPaths[Attempt];
+		if (FPackageName::DoesPackageExist(Path))
 		{
-			bCheckedRealtimeMode = true;
-			if (UE::GameFeatures::RealtimeMode)
+			GameFeatureDataHandle = UGameFeaturesSubsystem::LoadGameFeatureData(Path, true /*bStartStalled*/);
+			if (GameFeatureDataHandle && GameFeatureDataHandle->IsLoadingInProgress())
 			{
-				UE::GameFeatures::RealtimeMode->AddUpdateRequest(StateProperties.OnRequestUpdateStateMachine);
-				return;
+				GameFeatureDataHandle->BindCancelDelegate(FStreamableDelegate::CreateLambda([this]
+				{
+					const FStringView ShortUrl = StateProperties.PluginIdentifier.GetIdentifyingString();
+					UE_LOG(LogGameFeatures, Error, TEXT("Game Feature Data loading was cancelled for URL %.*s"), ShortUrl.Len(), ShortUrl.GetData());
+
+					LoadGFDState = ELoadGFDState::Cancelled;
+					UpdateStateMachineDeferred();
+				}));
+
+				GameFeatureDataHandle->BindCompleteDelegate(FStreamableDelegate::CreateLambda([this, Attempt]
+				{
+					StateProperties.GameFeatureData = Cast<UGameFeatureData>(GameFeatureDataHandle->GetLoadedAsset());
+					if (!StateProperties.GameFeatureData)
+					{
+						TryAsyncLoadGameFeatureData(Attempt + 1);
+						return;
+					}
+
+					LoadGFDState = ELoadGFDState::Success;
+					UpdateStateMachineDeferred();
+				}));
+
+				bIsLoading = true;
+				GameFeatureDataHandle->StartStalledHandle();
 			}
 		}
 
-		TRACE_CPUPROFILER_EVENT_SCOPE(GFP_Registering);
+		if (!bIsLoading)
+		{
+			TryAsyncLoadGameFeatureData(Attempt + 1);
+		}
+	}
+
+	virtual void BeginState() override
+	{
+		TRACE_CPUPROFILER_EVENT_SCOPE(GFP_Registering_Begin);
+
+		bCheckedRealtimeMode = false;
+
 		const FString PluginFolder = FPaths::GetPath(StateProperties.PluginInstalledFilename);
-		UGameplayTagsManager::Get().AddTagIniSearchPath(PluginFolder / TEXT("Config") / TEXT("Tags"));
 
-		const FString BackupGameFeatureDataPath = FString::Printf(TEXT("/%s/%s.%s"), *StateProperties.PluginName, *StateProperties.PluginName, *StateProperties.PluginName);
+		if (AllowIniLoading())
+		{
+			UGameplayTagsManager::Get().AddTagIniSearchPath(PluginFolder / TEXT("Config") / TEXT("Tags"));
+		}
+
+		LoadGFDState = ELoadGFDState::Pending;
+
+		TSharedPtr<IPlugin> Plugin = IPluginManager::Get().FindPlugin(StateProperties.PluginName);
+		ensure(Plugin.IsValid());
+
+		// If the plugin contains content then load the GameFeatureData otherwise procedurally create one that is transient.
+		if (!Plugin->GetDescriptor().bCanContainContent)
+		{
+			StateProperties.GameFeatureData = NewObject<UGameFeatureData>(GetTransientPackage(), FName(*StateProperties.PluginName), RF_Transient);
+			LoadGFDState = ELoadGFDState::Success;
+			return;
+		}
+
+		FString BackupGameFeatureDataPath = FString::Printf(TEXT("/%s/%s.%s"), *StateProperties.PluginName, *StateProperties.PluginName, *StateProperties.PluginName);
 
 		FString PreferredGameFeatureDataPath = TEXT("/") + StateProperties.PluginName + TEXT("/GameFeatureData.GameFeatureData");
 
@@ -2221,46 +2297,81 @@ struct FGameFeaturePluginState_Registering : public FGameFeaturePluginState
 			}
 			if (!OverridePath.IsEmpty())
 			{
-				PreferredGameFeatureDataPath = OverridePath;
+				PreferredGameFeatureDataPath = MoveTemp(OverridePath);
 			}
 		}
-		
-		auto LoadGameFeatureData = [](const FString& Path) -> UGameFeatureData*
+
+		GameFeatureDataSearchPaths.Empty();
+		GameFeatureDataSearchPaths.Emplace(MoveTemp(PreferredGameFeatureDataPath));
+		GameFeatureDataSearchPaths.Emplace(MoveTemp(BackupGameFeatureDataPath));
+
+		if (UseAsyncLoading())
 		{
-			TSharedPtr<FStreamableHandle> GameFeatureDataHandle;
-			if (FPackageName::DoesPackageExist(Path))
+			TryAsyncLoadGameFeatureData();
+		}
+		else
+		{
+			for (const FString& Path : GameFeatureDataSearchPaths)
 			{
-				GameFeatureDataHandle = UGameFeaturesSubsystem::LoadGameFeatureData(Path);
-				// @todo make this async. For now we just wait
-				if (GameFeatureDataHandle.IsValid())
+				if (FPackageName::DoesPackageExist(Path))
+				{
+					GameFeatureDataHandle = UGameFeaturesSubsystem::LoadGameFeatureData(Path);
+					if (GameFeatureDataHandle)
+					{
+						GameFeatureDataHandle->WaitUntilComplete(0.0f, false);
+						StateProperties.GameFeatureData = Cast<UGameFeatureData>(GameFeatureDataHandle->GetLoadedAsset());
+					}
+				}
+
+				if (StateProperties.GameFeatureData)
 				{
-					GameFeatureDataHandle->WaitUntilComplete(0.0f, false);
-					return Cast<UGameFeatureData>(GameFeatureDataHandle->GetLoadedAsset());
+					break;
 				}
 			}
 
-			return nullptr;
-		};
+			LoadGFDState = StateProperties.GameFeatureData ? ELoadGFDState::Success : ELoadGFDState::Failed;
+		}
+	}
+
+	virtual void EndState() override
+	{
+		GameFeatureDataHandle = nullptr;
+	}
 
-		TSharedPtr<IPlugin> Plugin = IPluginManager::Get().FindPlugin(StateProperties.PluginName);
-		ensure(Plugin.IsValid());
-		
-		// If the plugin contains content then load the GameFeatureData otherwise procedurally create one that is transient.
-		if (Plugin->GetDescriptor().bCanContainContent)
+	virtual void UpdateState(FGameFeaturePluginStateStatus& StateStatus) override
+	{
+		TRACE_CPUPROFILER_EVENT_SCOPE(GFP_Registering_Update);
+
+		if (!bCheckedRealtimeMode)
 		{
-			StateProperties.GameFeatureData = LoadGameFeatureData(PreferredGameFeatureDataPath);
-			if (!StateProperties.GameFeatureData)
+			bCheckedRealtimeMode = true;
+			if (UE::GameFeatures::RealtimeMode)
 			{
-				StateProperties.GameFeatureData = LoadGameFeatureData(BackupGameFeatureDataPath);
+				UE::GameFeatures::RealtimeMode->AddUpdateRequest(StateProperties.OnRequestUpdateStateMachine);
+				return;
 			}
 		}
-		else
+
+		if (!StateProperties.GameFeatureData)
 		{
-			StateProperties.GameFeatureData = NewObject<UGameFeatureData>(GetTransientPackage(), FName(*StateProperties.PluginName), RF_Transient);
+			check(LoadGFDState != ELoadGFDState::Success);
+
+			if (LoadGFDState == ELoadGFDState::Pending)
+			{
+				return;
+			}
+
+			if (LoadGFDState == ELoadGFDState::Cancelled)
+			{
+				StateStatus.SetTransitionError(EGameFeaturePluginState::ErrorRegistering, GetErrorResult(TEXT("Load_Cancelled_GameFeatureData")));
+				return;
+			}
 		}
 
 		if (StateProperties.GameFeatureData)
 		{
+			check(LoadGFDState == ELoadGFDState::Success);
+
 			if (AllowIniLoading())
 			{
 				StateProperties.GameFeatureData->InitializeBasePluginIniFile(StateProperties.PluginInstalledFilename);
@@ -2275,6 +2386,8 @@ struct FGameFeaturePluginState_Registering : public FGameFeaturePluginState
 		}
 		else
 		{
+			check(LoadGFDState == ELoadGFDState::Failed);
+
 			// The gamefeaturedata does not exist. The pak file may not be openable or this is a builtin plugin where the pak file does not exist.
 			StateStatus.SetTransitionError(EGameFeaturePluginState::ErrorRegistering, GetErrorResult(TEXT("Plugin_Missing_GameFeatureData")));
 		}
@@ -2298,6 +2411,23 @@ struct FGameFeaturePluginState_Registered : public FDestinationGameFeaturePlugin
 	}
 };
 
+struct FGameFeaturePluginState_ErrorLoading : public FErrorGameFeaturePluginState
+{
+	FGameFeaturePluginState_ErrorLoading(FGameFeaturePluginStateMachineProperties& InStateProperties) : FErrorGameFeaturePluginState(InStateProperties) {}
+
+	virtual void UpdateState(FGameFeaturePluginStateStatus& StateStatus) override
+	{
+		if (StateProperties.Destination < EGameFeaturePluginState::ErrorLoading)
+		{
+			StateStatus.SetTransition(EGameFeaturePluginState::Registered);
+		}
+		else if (StateProperties.Destination > EGameFeaturePluginState::ErrorLoading)
+		{
+			StateStatus.SetTransition(EGameFeaturePluginState::Loading);
+		}
+	}
+};
+
 struct FGameFeaturePluginState_Unloading : public FGameFeaturePluginState
 {
 	FGameFeaturePluginState_Unloading(FGameFeaturePluginStateMachineProperties& InStateProperties) : FGameFeaturePluginState(InStateProperties) {}
@@ -2363,17 +2493,57 @@ struct FGameFeaturePluginState_Loading : public FGameFeaturePluginState
 {
 	FGameFeaturePluginState_Loading(FGameFeaturePluginStateMachineProperties& InStateProperties) : FGameFeaturePluginState(InStateProperties) {}
 
+	TSharedPtr<FStreamableHandle> BundleHandle;
+
+	virtual void BeginState() override
+	{
+		TRACE_CPUPROFILER_EVENT_SCOPE(GFP_Loading_Begin);
+		check(StateProperties.GameFeatureData);
+
+		BundleHandle = LoadGameFeatureBundles(StateProperties.GameFeatureData);
+		if (BundleHandle)
+		{
+			// This will only bind if a load is in progress
+			BundleHandle->BindCancelDelegate(FStreamableDelegate::CreateLambda([this]
+			{
+				const FStringView ShortUrl = StateProperties.PluginIdentifier.GetIdentifyingString();
+				UE_LOG(LogGameFeatures, Error, TEXT("Game Feature preloading was cancelled for URL %.*s"), ShortUrl.Len(), ShortUrl.GetData());
+				UpdateStateMachineDeferred();
+			}));
+
+			// This will only bind if a load is in progress
+			BundleHandle->BindCompleteDelegate(FStreamableDelegate::CreateRaw(this, &FGameFeaturePluginState_Registering::UpdateStateMachineDeferred, 0.0f));
+		}
+	}
+
+	virtual void EndState() override
+	{
+		BundleHandle = nullptr;
+	}
+
 	virtual void UpdateState(FGameFeaturePluginStateStatus& StateStatus) override
 	{
-		TRACE_CPUPROFILER_EVENT_SCOPE(GFP_Loading);
+		TRACE_CPUPROFILER_EVENT_SCOPE(GFP_Loading_Update);
 		check(StateProperties.GameFeatureData);
 
-		// AssetManager
-		TSharedPtr<FStreamableHandle> BundleHandle = LoadGameFeatureBundles(StateProperties.GameFeatureData);
-		// @todo make this async. For now we just wait
-		if (BundleHandle.IsValid())
+		if (BundleHandle)
 		{
-			BundleHandle->WaitUntilComplete(0.0f, false);
+			if (!UseAsyncLoading())
+			{
+				BundleHandle->WaitUntilComplete(0.0f, false);
+			}
+
+			if (BundleHandle->IsLoadingInProgress())
+			{
+				return;
+			}
+
+			if (BundleHandle->WasCanceled())
+			{
+				BundleHandle.Reset();
+				StateStatus.SetTransitionError(EGameFeaturePluginState::ErrorLoading, GetErrorResult(TEXT("Load_Cancelled_Preload")));
+				return;
+			}
 		}
 
 		UGameFeaturesSubsystem::Get().OnGameFeatureLoading(StateProperties.GameFeatureData, StateProperties.PluginIdentifier);
@@ -2582,6 +2752,7 @@ struct FGameFeaturePluginState_Activating : public FGameFeaturePluginState
 
 		UGameFeaturesSubsystem::Get().OnGameFeatureActivating(StateProperties.GameFeatureData, StateProperties.PluginName, Context, StateProperties.PluginIdentifier);
 
+		// @TODO: non-blocking wait here?
 		// If this plugin caused localization data to load, wait for that here before marking it as active
 		if (TSharedPtr<IPlugin> Plugin = IPluginManager::Get().FindPlugin(StateProperties.PluginName);
 			Plugin && Plugin->GetDescriptor().bExplicitlyLoaded && Plugin->GetDescriptor().LocalizationTargets.Num() > 0)
@@ -2959,6 +3130,8 @@ void UGameFeaturePluginStateMachine::UpdateStateMachine()
 		return;
 	}
 
+	TRACE_CPUPROFILER_EVENT_SCOPE(GFP_UpdateStateMachine);
+
 	TOptional<TGuardValue<bool>> ScopeGuard(InPlace, bInUpdateStateMachine, true);
 	
 	using StateIt = std::underlying_type<EGameFeaturePluginState>::type;
diff --git a/Engine/Plugins/Experimental/GameFeatures/Source/GameFeatures/Private/GameFeaturePluginStateMachine.h b/Engine/Plugins/Experimental/GameFeatures/Source/GameFeatures/Private/GameFeaturePluginStateMachine.h
--- a/Engine/Plugins/Experimental/GameFeatures/Source/GameFeatures/Private/GameFeaturePluginStateMachine.h
+++ b/Engine/Plugins/Experimental/GameFeatures/Source/GameFeatures/Private/GameFeaturePluginStateMachine.h
@@ -80,12 +80,12 @@ Transition states are expected to transition the machine to another state after
       +--+---------+      +-------------------+       +------v-------+   |
       |            |      |         !         |       |              |   |
       | Releasing  <------> ErrorManagingData <-------> Downloading  |   |
-      |            |      |                   |       |              |   |  
-      +--^---------+      +-------------------+       +-------+------+   |  
-         |                                                   |           | 
-         |                                                   |           | 
-         |     +-------------+                               |           |   
-         |     |      *      |                               v           |    
+      |            |      |                   |       |              |   |
+      +--^---------+      +-------------------+       +-------+------+   |
+         |                                                   |           |
+         |                                                   |           |
+         |     +-------------+                               |           |
+         |     |      *      |                               v           |
          ------+ Installed   <--------------------------------------------
                |             |
                +-^---------+-+
@@ -120,13 +120,13 @@ Transition states are expected to transition the machine to another state after
                |            |
                +-^--------+-+
                  |        |
-           ------~--------~------------------------------
-           |     |        |                             |
-        +--v-----+--+  +--v--------+                    |
-        |           |  |           |                    |
-        | Unloading |  |  Loading  |                    |
-        |           |  |           |                    |
-        +--------^--+  +--+--------+                    |
+           ------~--------~-------------------------------------
+           |     |        |                             ^      |
+        +--v-----+--+  +--v--------+                    |    +-+------------+
+        |           |  |           |                    |    |      !       |
+        | Unloading |  |  Loading  <--------------------~----> ErrorLoading |
+        |           |  |           |                    |    |              |
+        +--------^--+  +--+--------+                    |    +--------------+
                  |        |                             |
                +-+--------v-+                           |
                |      *     |                           |
@@ -398,6 +398,8 @@ protected:
 
 	bool AllowIniLoading() const;
 
+	bool UseAsyncLoading() const;
+
 private:
 	void CleanupDeferredUpdateCallbacks() const;
 
diff --git a/Engine/Plugins/Experimental/GameFeatures/Source/GameFeatures/Private/GameFeaturesSubsystem.cpp b/Engine/Plugins/Experimental/GameFeatures/Source/GameFeatures/Private/GameFeaturesSubsystem.cpp
--- a/Engine/Plugins/Experimental/GameFeatures/Source/GameFeatures/Private/GameFeaturesSubsystem.cpp
+++ b/Engine/Plugins/Experimental/GameFeatures/Source/GameFeatures/Private/GameFeaturesSubsystem.cpp
@@ -491,9 +491,15 @@ void UGameFeaturesSubsystem::OnAssetManagerCreated()
 	bInitializedPolicyManager = true;
 }
 
-TSharedPtr<FStreamableHandle> UGameFeaturesSubsystem::LoadGameFeatureData(const FString& GameFeatureToLoad)
-{
-	return UAssetManager::Get().GetStreamableManager().RequestAsyncLoad(FSoftObjectPath(GameFeatureToLoad));
+TSharedPtr<FStreamableHandle> UGameFeaturesSubsystem::LoadGameFeatureData(const FString& GameFeatureToLoad, bool bStartStalled /*= false*/)
+{
+	return UAssetManager::Get().GetStreamableManager().RequestAsyncLoad(
+		FSoftObjectPath(GameFeatureToLoad), 
+		FStreamableDelegate(), 
+		FStreamableManager::DefaultAsyncLoadPriority, 
+		false, 
+		bStartStalled, 
+		TEXT("LoadGameFeatureData"));
 }
 
 void UGameFeaturesSubsystem::UnloadGameFeatureData(const UGameFeatureData* GameFeatureToUnload)
diff --git a/Engine/Plugins/Experimental/GameFeatures/Source/GameFeatures/Public/GameFeatureTypes.h b/Engine/Plugins/Experimental/GameFeatures/Source/GameFeatures/Public/GameFeatureTypes.h
--- a/Engine/Plugins/Experimental/GameFeatures/Source/GameFeatures/Public/GameFeatureTypes.h
+++ b/Engine/Plugins/Experimental/GameFeatures/Source/GameFeatures/Public/GameFeatureTypes.h
@@ -30,6 +30,7 @@ class FString;
 	XSTATE(Unregistering,				NSLOCTEXT("GameFeatures", "UnregisteringStateDisplayName", "Unregistering"))								/* Transition state Registered -> Installed. Cleaning up data gathered in Registering. */ \
 	XSTATE(Registering,					NSLOCTEXT("GameFeatures", "RegisteringStateDisplayName", "Registering"))									/* Transition state Installed -> Registered. Discovering assets in the plugin, but not loading them, except a few for discovery reasons. */ \
 	XSTATE(Registered,					NSLOCTEXT("GameFeatures", "RegisteredStateDisplayName", "Registered"))										/* The assets in the plugin are known, but have not yet been loaded, except a few for discovery reasons. */ \
+	XSTATE(ErrorLoading,				NSLOCTEXT("GameFeatures", "ErrorLoadingDisplayName", "ErrorLoading"))										/* Error state for Loading -> Loaded transition */ \
 	XSTATE(Unloading,					NSLOCTEXT("GameFeatures", "UnloadingStateDisplayName", "Unloading"))										/* Transition state Loaded -> Registered. In the process of removing code/content from memory. */ \
 	XSTATE(Loading,						NSLOCTEXT("GameFeatures", "LoadingStateDisplayName", "Loading"))											/* Transition state Registered -> Loaded. In the process of loading code/content into memory. */ \
 	XSTATE(Loaded,						NSLOCTEXT("GameFeatures", "LoadedStateDisplayName", "Loaded"))												/* The plugin is loaded into memory, but not registered with game systems and active. */ \
diff --git a/Engine/Plugins/Experimental/GameFeatures/Source/GameFeatures/Public/GameFeaturesSubsystem.h b/Engine/Plugins/Experimental/GameFeatures/Source/GameFeatures/Public/GameFeaturesSubsystem.h
--- a/Engine/Plugins/Experimental/GameFeatures/Source/GameFeatures/Public/GameFeaturesSubsystem.h
+++ b/Engine/Plugins/Experimental/GameFeatures/Source/GameFeatures/Public/GameFeaturesSubsystem.h
@@ -371,7 +371,7 @@ public:
 
 public:
 	/** Loads the specified game feature data and its bundles */
-	static TSharedPtr<FStreamableHandle> LoadGameFeatureData(const FString& GameFeatureToLoad);
+	static TSharedPtr<FStreamableHandle> LoadGameFeatureData(const FString& GameFeatureToLoad, bool bStartStalled = false);
 	static void UnloadGameFeatureData(const UGameFeatureData* GameFeatureToUnload);
 
 	void AddObserver(UObject* Observer);
diff --git a/Engine/Plugins/Experimental/GameFeatures/Source/GameFeaturesEditor/Private/SGameFeatureStateWidget.cpp b/Engine/Plugins/Experimental/GameFeatures/Source/GameFeaturesEditor/Private/SGameFeatureStateWidget.cpp
--- a/Engine/Plugins/Experimental/GameFeatures/Source/GameFeaturesEditor/Private/SGameFeatureStateWidget.cpp
+++ b/Engine/Plugins/Experimental/GameFeatures/Source/GameFeaturesEditor/Private/SGameFeatureStateWidget.cpp
@@ -65,7 +65,7 @@ FText SGameFeatureStateWidget::GetDisplayNameOfState(EGameFeaturePluginState Sta
 
 FText SGameFeatureStateWidget::GetTooltipOfState(EGameFeaturePluginState State)
 {
-	static_assert((int32)EGameFeaturePluginState::MAX == 31, "");
+	static_assert((int32)EGameFeaturePluginState::MAX == 32, "");
 
 	switch (State)
 	{
@@ -115,6 +115,8 @@ FText SGameFeatureStateWidget::GetTooltipOfState(EGameFeaturePluginState State)
 		return LOCTEXT("StateTooltip_Registering", "Transition state Installed -> Registered. Discovering assets in the plugin, but not loading them, except a few for discovery reasons.");
 	case EGameFeaturePluginState::Registered:
 		return LOCTEXT("StateTooltip_Registered", "The assets in the plugin are known, but have not yet been loaded, except a few for discovery reasons.");
+	case EGameFeaturePluginState::ErrorLoading:
+		return LOCTEXT("StateTooltip_ErrorLoading", "Error state for Loading -> Loaded transition.");
 	case EGameFeaturePluginState::Unloading:
 		return LOCTEXT("StateTooltip_Unloading", "Transition state Loaded -> Registered. In the process of removing code/content from memory.");
 	case EGameFeaturePluginState::Loading:
