Additional log to output the total time spent loading built in plugins

diff --git a/Engine/Plugins/Experimental/GameFeatures/Source/GameFeatures/Private/GameFeaturesSubsystem.cpp b/Engine/Plugins/Experimental/GameFeatures/Source/GameFeatures/Private/GameFeaturesSubsystem.cpp
--- a/Engine/Plugins/Experimental/GameFeatures/Source/GameFeatures/Private/GameFeaturesSubsystem.cpp
+++ b/Engine/Plugins/Experimental/GameFeatures/Source/GameFeatures/Private/GameFeaturesSubsystem.cpp
@@ -1504,8 +1504,16 @@ public:
 class FBuiltInPluginLoadTimeTracker
 {
 public:
+	FBuiltInPluginLoadTimeTracker()
+	{
+		StartTime = FPlatformTime::Seconds();
+	}
+
 	~FBuiltInPluginLoadTimeTracker()
 	{
+		double TotalLoadTime = FPlatformTime::Seconds() - StartTime;
+		UE_LOG(LogGameFeatures, Display, TEXT("Total built in plugin load time %.4fs"), TotalLoadTime);
+
 		if (PluginLoadTimes.Num() > 0)
 		{
 			UE_LOG(LogGameFeatures, Display, TEXT("There were %d built in plugins that took longer than %.4fs to load. Listing worst offenders."), PluginLoadTimes.Num(), UE::GameFeatures::GBuiltInPluginLoadTimeReportThreshold);
@@ -1535,7 +1543,7 @@ public:
 
 private:
 	TArray<TPair<FString, double>> PluginLoadTimes;
-	
+	double StartTime;
 };
 
 class FBuiltInPluginLoadTimeTrackerScope
