Expose parsed JsonObject so it can be used when determing if a plugin should be loaded

diff --git a/Engine/Plugins/Experimental/GameFeatures/Source/GameFeatures/Private/GameFeaturesSubsystem.cpp b/Engine/Plugins/Experimental/GameFeatures/Source/GameFeatures/Private/GameFeaturesSubsystem.cpp
--- a/Engine/Plugins/Experimental/GameFeatures/Source/GameFeatures/Private/GameFeaturesSubsystem.cpp
+++ b/Engine/Plugins/Experimental/GameFeatures/Source/GameFeatures/Private/GameFeaturesSubsystem.cpp
@@ -1809,7 +1809,7 @@ bool UGameFeaturesSubsystem::GetGameFeaturePluginDetails(const FString& PluginUR
 	}
 
 	// Read the properties
-
+	OutPluginDetails.CachedJson = ObjectPtr;
 	// Hotfixable. If it is not specified, then we assume it is
 	OutPluginDetails.bHotfixable = true;
 	ObjectPtr->TryGetBoolField(TEXT("Hotfixable"), OutPluginDetails.bHotfixable);
diff --git a/Engine/Plugins/Experimental/GameFeatures/Source/GameFeatures/Public/GameFeaturesSubsystem.h b/Engine/Plugins/Experimental/GameFeatures/Source/GameFeatures/Public/GameFeaturesSubsystem.h
--- a/Engine/Plugins/Experimental/GameFeatures/Source/GameFeatures/Public/GameFeaturesSubsystem.h
+++ b/Engine/Plugins/Experimental/GameFeatures/Source/GameFeatures/Public/GameFeaturesSubsystem.h
@@ -244,6 +244,7 @@ struct FGameFeaturePluginDetails
 {
 	TArray<FGameFeaturePluginReferenceDetails> PluginDependencies;
 	TMap<FString, FString> AdditionalMetadata;
+	TSharedPtr<class FJsonObject> CachedJson;
 	bool bHotfixable;
 	EBuiltInAutoState BuiltInAutoState;
 
