Fix wrong state being bound in member function delegate

diff --git a/Engine/Plugins/Experimental/GameFeatures/Source/GameFeatures/Private/GameFeaturePluginStateMachine.cpp b/Engine/Plugins/Experimental/GameFeatures/Source/GameFeatures/Private/GameFeaturePluginStateMachine.cpp
--- a/Engine/Plugins/Experimental/GameFeatures/Source/GameFeatures/Private/GameFeaturePluginStateMachine.cpp
+++ b/Engine/Plugins/Experimental/GameFeatures/Source/GameFeatures/Private/GameFeaturePluginStateMachine.cpp
@@ -2512,7 +2512,7 @@ struct FGameFeaturePluginState_Loading : public FGameFeaturePluginState
 			}));
 
 			// This will only bind if a load is in progress
-			BundleHandle->BindCompleteDelegate(FStreamableDelegate::CreateRaw(this, &FGameFeaturePluginState_Registering::UpdateStateMachineDeferred, 0.0f));
+			BundleHandle->BindCompleteDelegate(FStreamableDelegate::CreateRaw(this, &FGameFeaturePluginState_Loading::UpdateStateMachineDeferred, 0.0f));
 		}
 	}
 
