Convert AdditionalMetaData to be a FJsonValue so that more complex structures could be sampled for.

diff --git a/Engine/Plugins/Experimental/GameFeatures/Source/GameFeatures/Private/GameFeaturesSubsystem.cpp b/Engine/Plugins/Experimental/GameFeatures/Source/GameFeatures/Private/GameFeaturesSubsystem.cpp
--- a/Engine/Plugins/Experimental/GameFeatures/Source/GameFeatures/Private/GameFeaturesSubsystem.cpp
+++ b/Engine/Plugins/Experimental/GameFeatures/Source/GameFeatures/Private/GameFeaturesSubsystem.cpp
@@ -4,6 +4,7 @@
 #include "Algo/TopologicalSort.h"
 #include "AssetRegistry/AssetData.h"
 #include "AssetRegistry/IAssetRegistry.h"
+#include "Dom/JsonValue.h"
 #include "GameFeaturesSubsystemSettings.h"
 #include "GameFeaturesProjectPolicies.h"
 #include "GameFeatureData.h"
@@ -1809,7 +1810,6 @@ bool UGameFeaturesSubsystem::GetGameFeaturePluginDetails(const FString& PluginUR
 	}
 
 	// Read the properties
-	OutPluginDetails.CachedJson = ObjectPtr;
 	// Hotfixable. If it is not specified, then we assume it is
 	OutPluginDetails.bHotfixable = true;
 	ObjectPtr->TryGetBoolField(TEXT("Hotfixable"), OutPluginDetails.bHotfixable);
@@ -1820,9 +1820,15 @@ bool UGameFeaturesSubsystem::GetGameFeaturePluginDetails(const FString& PluginUR
 	// Read any additional metadata the policy might want to consume (e.g., a release version number)
 	for (const FString& ExtraKey : GetDefault<UGameFeaturesSubsystemSettings>()->AdditionalPluginMetadataKeys)
 	{
-		FString ExtraValue;
-		ObjectPtr->TryGetStringField(ExtraKey, ExtraValue);
-		OutPluginDetails.AdditionalMetadata.Add(ExtraKey, ExtraValue);
+		TSharedPtr<FJsonValue> Field = ObjectPtr->TryGetField(ExtraKey);
+		if (Field.IsValid())
+		{
+			OutPluginDetails.AdditionalMetadata.Add(ExtraKey, Field);
+		}
+		else
+		{
+			OutPluginDetails.AdditionalMetadata.Add(ExtraKey, MakeShared<FJsonValueString>(TEXT("")));
+		}
 	}
 
 	// Parse plugin dependencies
diff --git a/Engine/Plugins/Experimental/GameFeatures/Source/GameFeatures/Public/GameFeaturesSubsystem.h b/Engine/Plugins/Experimental/GameFeatures/Source/GameFeatures/Public/GameFeaturesSubsystem.h
--- a/Engine/Plugins/Experimental/GameFeatures/Source/GameFeatures/Public/GameFeaturesSubsystem.h
+++ b/Engine/Plugins/Experimental/GameFeatures/Source/GameFeatures/Public/GameFeaturesSubsystem.h
@@ -243,8 +243,7 @@ struct FGameFeaturePluginReferenceDetails
 struct FGameFeaturePluginDetails
 {
 	TArray<FGameFeaturePluginReferenceDetails> PluginDependencies;
-	TMap<FString, FString> AdditionalMetadata;
-	TSharedPtr<class FJsonObject> CachedJson;
+	TMap<FString, TSharedPtr<FJsonValue>> AdditionalMetadata;
 	bool bHotfixable;
 	EBuiltInAutoState BuiltInAutoState;
 
