Non-unity

diff --git a/Engine/Plugins/Experimental/GameFeatures/Source/GameFeatures/Public/GameFeaturesSubsystem.h b/Engine/Plugins/Experimental/GameFeatures/Source/GameFeatures/Public/GameFeaturesSubsystem.h
--- a/Engine/Plugins/Experimental/GameFeatures/Source/GameFeatures/Public/GameFeaturesSubsystem.h
+++ b/Engine/Plugins/Experimental/GameFeatures/Source/GameFeatures/Public/GameFeaturesSubsystem.h
@@ -243,7 +243,7 @@ struct FGameFeaturePluginReferenceDetails
 struct FGameFeaturePluginDetails
 {
 	TArray<FGameFeaturePluginReferenceDetails> PluginDependencies;
-	TMap<FString, TSharedPtr<FJsonValue>> AdditionalMetadata;
+	TMap<FString, TSharedPtr<class FJsonValue>> AdditionalMetadata;
 	bool bHotfixable;
 	EBuiltInAutoState BuiltInAutoState;
 
