Fix wrong state transition from FGameFeaturePluginState_ErrorLoading when unloading.

diff --git a/Engine/Plugins/Experimental/GameFeatures/Source/GameFeatures/Private/GameFeaturePluginStateMachine.cpp b/Engine/Plugins/Experimental/GameFeatures/Source/GameFeatures/Private/GameFeaturePluginStateMachine.cpp
--- a/Engine/Plugins/Experimental/GameFeatures/Source/GameFeatures/Private/GameFeaturePluginStateMachine.cpp
+++ b/Engine/Plugins/Experimental/GameFeatures/Source/GameFeatures/Private/GameFeaturePluginStateMachine.cpp
@@ -2419,7 +2419,7 @@ struct FGameFeaturePluginState_ErrorLoading : public FErrorGameFeaturePluginStat
 	{
 		if (StateProperties.Destination < EGameFeaturePluginState::ErrorLoading)
 		{
-			StateStatus.SetTransition(EGameFeaturePluginState::Registered);
+			StateStatus.SetTransition(EGameFeaturePluginState::Unloading);
 		}
 		else if (StateProperties.Destination > EGameFeaturePluginState::ErrorLoading)
 		{
