Engine - 26.30 - Non-const getter to allow classes running in editor to modify actions list in game feature data

diff --git a/Engine/Plugins/Experimental/GameFeatures/Source/GameFeatures/Public/GameFeatureData.h b/Engine/Plugins/Experimental/GameFeatures/Source/GameFeatures/Public/GameFeatureData.h
--- a/Engine/Plugins/Experimental/GameFeatures/Source/GameFeatures/Public/GameFeatureData.h
+++ b/Engine/Plugins/Experimental/GameFeatures/Source/GameFeatures/Public/GameFeatureData.h
@@ -65,6 +65,10 @@ public:
 
 	const TArray<UGameFeatureAction*>& GetActions() const { return Actions; }
 
+#if WITH_EDITOR
+	TArray<UGameFeatureAction*>& GetMutableActionsInEditor() { return Actions; }
+#endif
+
 private:
 	/** Internal helper function to reload config data on objects as a result of a plugin INI being loaded */
 	void ReloadConfigs(FConfigFile& PluginConfig) const;
