Engine - 26.30 - Mac target fix for Change 27947576

diff --git a/Engine/Plugins/Experimental/GameFeatures/Source/GameFeatures/Public/GameFeatureData.h b/Engine/Plugins/Experimental/GameFeatures/Source/GameFeatures/Public/GameFeatureData.h
--- a/Engine/Plugins/Experimental/GameFeatures/Source/GameFeatures/Public/GameFeatureData.h
+++ b/Engine/Plugins/Experimental/GameFeatures/Source/GameFeatures/Public/GameFeatureData.h
@@ -66,7 +66,7 @@ public:
 	const TArray<UGameFeatureAction*>& GetActions() const { return Actions; }
 
 #if WITH_EDITOR
-	TArray<UGameFeatureAction*>& GetMutableActionsInEditor() { return Actions; }
+	TArray<TObjectPtr<UGameFeatureAction>>& GetMutableActionsInEditor() { return Actions; }
 #endif
 
 private:
