[GFS] Add realtime mode support for Unloading state
#rnx

diff --git a/Engine/Plugins/Experimental/GameFeatures/Source/GameFeatures/Private/GameFeaturePluginStateMachine.cpp b/Engine/Plugins/Experimental/GameFeatures/Source/GameFeatures/Private/GameFeaturePluginStateMachine.cpp
--- a/Engine/Plugins/Experimental/GameFeatures/Source/GameFeatures/Private/GameFeaturePluginStateMachine.cpp
+++ b/Engine/Plugins/Experimental/GameFeatures/Source/GameFeatures/Private/GameFeaturePluginStateMachine.cpp
@@ -1684,42 +1684,13 @@ struct FGameFeaturePluginState_Unmounting : public FGameFeaturePluginState
 
 	UE::GameFeatures::FResult Result = MakeValue();
 	TArray<FName> PendingBundles;
+	bool bUnmounting = false;
 	bool bUnmounted = false;
+	bool bCheckedRealtimeMode = false;
 
-	void OnContentReleased(FInstallBundleReleaseRequestResultInfo BundleResult)
-	{
-		if (!PendingBundles.Contains(BundleResult.BundleName))
-		{
-			return;
-		}
-
-		PendingBundles.Remove(BundleResult.BundleName);
-
-		if (!Result.HasError() && BundleResult.Result != EInstallBundleReleaseResult::OK)
-		{
-			Result = GetErrorResult(TEXT("BundleManager.OnReleased."), BundleResult.Result);
-		}
-
-		if (PendingBundles.Num() > 0)
-		{
-			return;
-		}
-
-		if (Result.HasValue())
-		{
-			bUnmounted = true;
-		}
-
-		UpdateStateMachineImmediate();
-	}
-
-	virtual void BeginState() override
+	void Unmount()
 	{
-		Result = MakeValue();
-		PendingBundles.Empty();
-		bUnmounted = false;
-
-		if (TSharedPtr<IPlugin> Plugin = IPluginManager::Get().FindPlugin(StateProperties.PluginName); 
+		if (TSharedPtr<IPlugin> Plugin = IPluginManager::Get().FindPlugin(StateProperties.PluginName);
 			Plugin && Plugin->GetDescriptor().bExplicitlyLoaded)
 		{
 			// The asset registry listens to FPackageName::OnContentPathDismounted() and 
@@ -1757,7 +1728,7 @@ struct FGameFeaturePluginState_Unmounting : public FGameFeaturePluginState
 		EInstallBundleReleaseRequestFlags ReleaseFlags = StateProperties.ProtocolOptions.GetSubtype<FInstallBundlePluginProtocolOptions>().ReleaseInstallBundleFlags;
 		//Make sure we don't remove files here early, that should only be done in Uninstalling
 		ReleaseFlags &= ~(EInstallBundleReleaseRequestFlags::RemoveFilesIfPossible);
-		
+
 		TValueOrError<FInstallBundleRequestInfo, EInstallBundleResult> MaybeRequestInfo = BundleManager->RequestReleaseContent(InstallBundles, ReleaseFlags);
 
 		if (MaybeRequestInfo.HasError())
@@ -1789,8 +1760,60 @@ struct FGameFeaturePluginState_Unmounting : public FGameFeaturePluginState
 		}
 	}
 
+	void OnContentReleased(FInstallBundleReleaseRequestResultInfo BundleResult)
+	{
+		if (!PendingBundles.Contains(BundleResult.BundleName))
+		{
+			return;
+		}
+
+		PendingBundles.Remove(BundleResult.BundleName);
+
+		if (!Result.HasError() && BundleResult.Result != EInstallBundleReleaseResult::OK)
+		{
+			Result = GetErrorResult(TEXT("BundleManager.OnReleased."), BundleResult.Result);
+		}
+
+		if (PendingBundles.Num() > 0)
+		{
+			return;
+		}
+
+		if (Result.HasValue())
+		{
+			bUnmounted = true;
+		}
+
+		UpdateStateMachineImmediate();
+	}
+
+	virtual void BeginState() override
+	{
+		Result = MakeValue();
+		PendingBundles.Empty();
+		bUnmounting = false;
+		bUnmounted = false;
+		bCheckedRealtimeMode = false;
+	}
+
 	virtual void UpdateState(FGameFeaturePluginStateStatus& StateStatus) override
 	{
+		if (!bCheckedRealtimeMode)
+		{
+			bCheckedRealtimeMode = true;
+			if (UE::GameFeatures::RealtimeMode)
+			{
+				UE::GameFeatures::RealtimeMode->AddUpdateRequest(StateProperties.OnRequestUpdateStateMachine);
+				return;
+			}
+		}
+
+		if (!bUnmounting)
+		{
+			bUnmounting = true;
+			Unmount();
+		}
+
 		if (!Result.HasValue())
 		{
 			StateStatus.SetTransitionError(EGameFeaturePluginState::ErrorMounting, Result);
