Implement async path for FGameFeaturePluginState_Mounting.  This addresses most game thread spikes that can slow down GFP installation although there is still a small spike during plugin mounting.

diff --git a/Engine/Plugins/Experimental/GameFeatures/Source/GameFeatures/Private/GameFeaturePluginStateMachine.cpp b/Engine/Plugins/Experimental/GameFeatures/Source/GameFeatures/Private/GameFeaturePluginStateMachine.cpp
--- a/Engine/Plugins/Experimental/GameFeatures/Source/GameFeatures/Private/GameFeaturePluginStateMachine.cpp
+++ b/Engine/Plugins/Experimental/GameFeatures/Source/GameFeatures/Private/GameFeaturePluginStateMachine.cpp
@@ -33,6 +33,7 @@
 #include "UObject/NoExportTypes.h"
 #include "Misc/PathViews.h"
 #include "Containers/Queue.h"
+#include "ShaderCodeLibrary.h"
 
 #include UE_INLINE_GENERATED_CPP_BY_NAME(GameFeaturePluginStateMachine)
 
@@ -1841,11 +1842,22 @@ struct FGameFeaturePluginState_Mounting : public FGameFeaturePluginState
 		, Result(MakeValue())
 	{}
 
+	enum class ESubState : uint8
+	{
+		None = 0,
+		MountPlugin,
+		LoadAssetRegistry
+	};
+	FRIEND_ENUM_CLASS_FLAGS(ESubState);
+
 	int32 NumObservedPostMountPausers = 0;
 	int32 NumExpectedPostMountPausers = 0;
 	TArray<FName> PendingBundles;
-	bool bCheckedRealtimeMode = false;
+	FDelegateHandle PakFileMountedDelegateHandle;
 	UE::GameFeatures::FResult Result;
+	ESubState StartedSubStates = ESubState::None;
+	ESubState CompletedSubStates = ESubState::None;
+	bool bCheckedRealtimeMode = false;
 
 	void OnInstallBundleCompleted(FInstallBundleRequestResultInfo BundleResult)
 	{
@@ -1871,7 +1883,11 @@ struct FGameFeaturePluginState_Mounting : public FGameFeaturePluginState
 		if (PendingBundles.IsEmpty())
 		{
 			IInstallBundleManager::InstallBundleCompleteDelegate.RemoveAll(this);
-			FCoreDelegates::GetOnPakFileMounted2().RemoveAll(this);
+			if (PakFileMountedDelegateHandle.IsValid())
+			{
+				FCoreDelegates::GetOnPakFileMounted2().Remove(PakFileMountedDelegateHandle);
+				PakFileMountedDelegateHandle.Reset();
+			}
 			
 			UpdateStateMachineImmediate();
 		}
@@ -1910,9 +1926,12 @@ struct FGameFeaturePluginState_Mounting : public FGameFeaturePluginState
 	{
 		NumObservedPostMountPausers = 0;
 		NumExpectedPostMountPausers = 0;
-		bCheckedRealtimeMode = false;
 		PendingBundles.Empty();
+		PakFileMountedDelegateHandle.Reset();
 		Result = MakeValue();
+		StartedSubStates = ESubState::None;
+		CompletedSubStates = ESubState::None;
+		bCheckedRealtimeMode = false;
 
 		if (StateProperties.GetPluginProtocol() != EGameFeaturePluginProtocol::InstallBundle)
 		{
@@ -1925,7 +1944,8 @@ struct FGameFeaturePluginState_Mounting : public FGameFeaturePluginState
 		const TArray<FName>& InstallBundles = MetaData.InstallBundles;
 
 		const FInstallBundlePluginProtocolOptions& Options = StateProperties.ProtocolOptions.GetSubtype<FInstallBundlePluginProtocolOptions>();
-		const EInstallBundleRequestFlags InstallFlags = Options.InstallBundleFlags /*| EInstallBundleRequestFlags::AsyncMount*/;
+		const EInstallBundleRequestFlags InstallFlags = UseAsyncLoading() ?
+			(Options.InstallBundleFlags | EInstallBundleRequestFlags::AsyncMount) : Options.InstallBundleFlags;
 
 		// Make bundle manager use verbose log level for most logs.
 		// We are already done with downloading, so we don't care about logging too much here unless mounting fails.
@@ -1956,116 +1976,220 @@ struct FGameFeaturePluginState_Mounting : public FGameFeaturePluginState
 			IInstallBundleManager::InstallBundleCompleteDelegate.AddRaw(this, &FGameFeaturePluginState_Mounting::OnInstallBundleCompleted);
 			if (UE::GameFeatures::ShouldLogMountedFiles)
 			{
-				FCoreDelegates::GetOnPakFileMounted2().AddRaw(this, &FGameFeaturePluginState_Mounting::OnPakFileMounted);
+				// Track with a delegate handle to avoid unbinding if we don't use this. Unbinding this causes an occaisonal perf spike.
+				PakFileMountedDelegateHandle = FCoreDelegates::GetOnPakFileMounted2().AddRaw(this, &FGameFeaturePluginState_Mounting::OnPakFileMounted);
 			}
 		}
 	}
 
-	virtual void UpdateState(FGameFeaturePluginStateStatus& StateStatus) override
+	void UpdateState_MountPlugin()
 	{
-		// Check if waiting for install bundles
-		if (PendingBundles.Num() > 0)
+		if (EnumHasAnyFlags(StartedSubStates, ESubState::MountPlugin))
 		{
 			return;
 		}
 
-		// Check if post-mount is paused
-		if (NumExpectedPostMountPausers > 0)
+		TRACE_CPUPROFILER_EVENT_SCOPE(GFP_Mounting_Plugin);
+
+		StartedSubStates |= ESubState::MountPlugin;
+
+		if (Result.HasError())
 		{
-			// Check if post-mount unpaused
-			if (NumExpectedPostMountPausers == NumObservedPostMountPausers)
-			{
-				NumExpectedPostMountPausers = INDEX_NONE;
-				TransitionOut(StateStatus);
-			}
+			CompletedSubStates |= ESubState::MountPlugin;
 			return;
 		}
 
-		if (!bCheckedRealtimeMode)
+		// Pre-mount
 		{
-			bCheckedRealtimeMode = true;
-			if (UE::GameFeatures::RealtimeMode)
-			{
-				UE::GameFeatures::RealtimeMode->AddUpdateRequest(StateProperties.OnRequestUpdateStateMachine);
-				return;
-			}
+			FGameFeaturePreMountingContext Context;
+			UGameFeaturesSubsystem::Get().OnGameFeaturePreMounting(StateProperties.PluginName, StateProperties.PluginIdentifier, Context);
 		}
 
-		TRACE_CPUPROFILER_EVENT_SCOPE(GFP_Mounting);
+		checkf(!StateProperties.PluginInstalledFilename.IsEmpty(), TEXT("PluginInstalledFilename must be set by the Mounting. PluginURL: %s"), *StateProperties.PluginIdentifier.GetFullPluginURL());
+		checkf(FPaths::GetExtension(StateProperties.PluginInstalledFilename) == TEXT("uplugin"), TEXT("PluginInstalledFilename must have a uplugin extension. PluginURL: %s"), *StateProperties.PluginIdentifier.GetFullPluginURL());
 
-		// Don't mount the plugin if there was an error during BeginState or bundles install
-		if (!Result.HasError())
+		// refresh the plugins list to let the plugin manager know about it
+		const TSharedPtr<IPlugin> MaybePlugin = IPluginManager::Get().FindPlugin(StateProperties.PluginName);
+		const bool bNeedsPluginMount = (MaybePlugin == nullptr || MaybePlugin->GetDescriptor().bExplicitlyLoaded);
+
+		if (MaybePlugin)
+		{
+			if (!FPaths::IsSamePath(MaybePlugin->GetDescriptorFileName(), StateProperties.PluginInstalledFilename))
+			{
+				Result = GetErrorResult(TEXT("Plugin_Name_Already_In_Use"));
+			}
+		}
+		else
 		{
-			// Pre-mount
+			const bool bAddedPlugin = IPluginManager::Get().AddToPluginsList(StateProperties.PluginInstalledFilename);
+			if (bAddedPlugin)
 			{
-				FGameFeaturePreMountingContext Context;
-				UGameFeaturesSubsystem::Get().OnGameFeaturePreMounting(StateProperties.PluginName, StateProperties.PluginIdentifier, Context);
+				StateProperties.bAddedPluginToManager = true;
 			}
+			else
+			{
+				Result = GetErrorResult(TEXT("Failed_To_Register_Plugin"));
+			}
+		}
+
+		if (Result.HasError() || !bNeedsPluginMount)
+		{
+			CompletedSubStates |= ESubState::MountPlugin;
+			return;
+		}
+
+		if (!UseAsyncLoading())
+		{
+			verify(IPluginManager::Get().MountExplicitlyLoadedPlugin(StateProperties.PluginName));
+			CompletedSubStates |= ESubState::MountPlugin;
+			return;
+		}
 
-			checkf(!StateProperties.PluginInstalledFilename.IsEmpty(), TEXT("PluginInstalledFilename must be set by the Mounting. PluginURL: %s"), *StateProperties.PluginIdentifier.GetFullPluginURL());
-			checkf(FPaths::GetExtension(StateProperties.PluginInstalledFilename) == TEXT("uplugin"), TEXT("PluginInstalledFilename must have a uplugin extension. PluginURL: %s"), *StateProperties.PluginIdentifier.GetFullPluginURL());
+		// We want to make sure opening the shader lib happens async
+		FShaderCodeLibrary::DontOpenPluginShaderLibraryOnMount(StateProperties.PluginName);
 
-			// refresh the plugins list to let the plugin manager know about it
-			const TSharedPtr<IPlugin> MaybePlugin = IPluginManager::Get().FindPlugin(StateProperties.PluginName);
-			const bool bNeedsPluginMount = (MaybePlugin == nullptr || MaybePlugin->GetDescriptor().bExplicitlyLoaded);
+		verify(IPluginManager::Get().MountExplicitlyLoadedPlugin(StateProperties.PluginName));
 
-			if (MaybePlugin)
+		// Now load the shader lib in the background
+		TSharedPtr<IPlugin> Plugin = IPluginManager::Get().FindPlugin(StateProperties.PluginName);
+		if (Plugin->CanContainContent() && Plugin->IsEnabled()) // TODO: possibly skip this if there is no shaderlib (need to figure out the file name)
+		{
+			UE::Tasks::Launch(UE_SOURCE_LOCATION, [this, PluginName=Plugin->GetName(), PluginDir=Plugin->GetContentDir()]
 			{
-				if (!FPaths::IsSamePath(MaybePlugin->GetDescriptorFileName(), StateProperties.PluginInstalledFilename))
+				FShaderCodeLibrary::OpenLibrary(PluginName, PluginDir);
+
+				ExecuteOnGameThread(UE_SOURCE_LOCATION, [this]
 				{
-					Result = GetErrorResult(TEXT("Plugin_Name_Already_In_Use"));
-				}
+					CompletedSubStates |= ESubState::MountPlugin;
+					UpdateStateMachineImmediate();
+				});
+
+			}, UE::Tasks::ETaskPriority::BackgroundHigh);
+
+			return;
+		}
+
+		CompletedSubStates |= ESubState::MountPlugin;
+	}
+
+	void UpdateState_LoadAssetRegistry()
+	{
+		if (EnumHasAnyFlags(StartedSubStates, ESubState::LoadAssetRegistry))
+		{
+			return;
+		}
+
+		StartedSubStates |= ESubState::LoadAssetRegistry;
+
+		if (Result.HasError())
+		{
+			CompletedSubStates |= ESubState::MountPlugin;
+			return;
+		}
+
+		if (StateProperties.GetPluginProtocol() != EGameFeaturePluginProtocol::InstallBundle)
+		{
+			CompletedSubStates |= ESubState::LoadAssetRegistry;
+			return;
+		}
+
+		TRACE_CPUPROFILER_EVENT_SCOPE(GFP_Mounting_AR);
+
+		// After the new plugin is mounted add the asset registry for that plugin.
+
+		const FString PluginFolder = FPaths::GetPath(StateProperties.PluginInstalledFilename);
+		FString PluginAssetRegistry = PluginFolder / TEXT("AssetRegistry.bin");
+
+		TSharedPtr<IPlugin> NewlyMountedPlugin = IPluginManager::Get().FindPlugin(StateProperties.PluginName);
+		if (!NewlyMountedPlugin || !NewlyMountedPlugin->CanContainContent() || !IFileManager::Get().FileExists(*PluginAssetRegistry))
+		{
+			CompletedSubStates |= ESubState::LoadAssetRegistry;
+			return;
+		}
+
+		if (!UseAsyncLoading())
+		{
+			FAssetRegistryState PluginAssetRegistryState;
+			if (FAssetRegistryState::LoadFromDisk(*PluginAssetRegistry, FAssetRegistryLoadOptions(), PluginAssetRegistryState))
+			{
+				IAssetRegistry& AssetRegistry = UAssetManager::Get().GetAssetRegistry();
+				AssetRegistry.AppendState(PluginAssetRegistryState);
 			}
 			else
 			{
-				const bool bAddedPlugin = IPluginManager::Get().AddToPluginsList(StateProperties.PluginInstalledFilename);
-				if (bAddedPlugin)
-				{
-					StateProperties.bAddedPluginToManager = true;
-				}
-				else
-				{
-					Result = GetErrorResult(TEXT("Failed_To_Register_Plugin"));
-				}
+				Result = GetErrorResult(TEXT("Failed_To_Load_Plugin_AssetRegistry"));
 			}
 
-			if (!Result.HasError())
+			CompletedSubStates |= ESubState::LoadAssetRegistry;
+			return;
+		}
+
+		UE::Tasks::Launch(UE_SOURCE_LOCATION, [this, PluginAssetRegistry=MoveTemp(PluginAssetRegistry)]
+		{
+			bool bSuccess = false;
+			TSharedPtr<FAssetRegistryState> PluginAssetRegistryState = MakeShared<FAssetRegistryState>();
+			if (FAssetRegistryState::LoadFromDisk(*PluginAssetRegistry, FAssetRegistryLoadOptions(), *PluginAssetRegistryState))
 			{
-				if (bNeedsPluginMount)
+				IAssetRegistry& AssetRegistry = UAssetManager::Get().GetAssetRegistry();
+				AssetRegistry.AppendState(*PluginAssetRegistryState);
+				bSuccess = true;
+			}
+
+			ExecuteOnGameThread(UE_SOURCE_LOCATION, [this, bSuccess]
+			{
+				TRACE_CPUPROFILER_EVENT_SCOPE(GFP_Mounting_ARComplete);
+
+				if (!bSuccess)
 				{
-					IPluginManager::Get().MountExplicitlyLoadedPlugin(StateProperties.PluginName);
+					Result = GetErrorResult(TEXT("Failed_To_Load_Plugin_AssetRegistry"));
 				}
 
-				// @TODO: Load AR Data async?
-				// After the new plugin is mounted add the asset registry for that plugin.
-				if (StateProperties.GetPluginProtocol() == EGameFeaturePluginProtocol::InstallBundle)
-				{
-					const FString PluginFolder = FPaths::GetPath(StateProperties.PluginInstalledFilename);
-					const FString PluginAssetRegistry = PluginFolder / TEXT("AssetRegistry.bin");
+				CompletedSubStates |= ESubState::LoadAssetRegistry;
+				UpdateStateMachineImmediate();
+			});
 
-					TSharedPtr<IPlugin> NewlyMountedPlugin = IPluginManager::Get().FindPlugin(StateProperties.PluginName);
-					if (NewlyMountedPlugin.IsValid() && NewlyMountedPlugin->CanContainContent() && IFileManager::Get().FileExists(*PluginAssetRegistry))
-					{
-						TArray<uint8> SerializedAssetData;
-						if (FFileHelper::LoadFileToArray(SerializedAssetData, *PluginAssetRegistry))
-						{
-							FAssetRegistryState PluginAssetRegistryState;
-							FMemoryReader Ar(SerializedAssetData);
-							PluginAssetRegistryState.Load(Ar);
+		}, UE::Tasks::ETaskPriority::BackgroundHigh);
+	}
 
-							IAssetRegistry& AssetRegistry = UAssetManager::Get().GetAssetRegistry();
-							AssetRegistry.AppendState(PluginAssetRegistryState);
-						}
-						else
-						{
-							Result = GetErrorResult(TEXT("Failed_To_Load_Plugin_AssetRegistry"));
-						}
-					}
-				}
+	virtual void UpdateState(FGameFeaturePluginStateStatus& StateStatus) override
+	{
+		// Check if waiting for install bundles
+		if (PendingBundles.Num() > 0)
+		{
+			return;
+		}
+
+		// Check if post-mount is paused
+		if (NumExpectedPostMountPausers > 0)
+		{
+			// Check if post-mount unpaused
+			if (NumExpectedPostMountPausers == NumObservedPostMountPausers)
+			{
+				NumExpectedPostMountPausers = INDEX_NONE;
+				TransitionOut(StateStatus);
 			}
+			return;
 		}
 
+		if (!bCheckedRealtimeMode)
+		{
+			bCheckedRealtimeMode = true;
+			if (UE::GameFeatures::RealtimeMode)
+			{
+				UE::GameFeatures::RealtimeMode->AddUpdateRequest(StateProperties.OnRequestUpdateStateMachine);
+				return;
+			}
+		}
+
+		TRACE_CPUPROFILER_EVENT_SCOPE(GFP_Mounting);
+		
+		UpdateState_MountPlugin();
+		UpdateState_LoadAssetRegistry();
+
+		const bool bComplete = EnumHasAllFlags(CompletedSubStates, ESubState::MountPlugin | ESubState::LoadAssetRegistry);
+
 		// Post-mount
+		if (bComplete)
 		{
 			FGameFeaturePostMountingContext Context(StateProperties.PluginName, [this](FStringView InPauserTag) { OnPostMountPauserCompleted(InPauserTag); });
 			NumExpectedPostMountPausers = INDEX_NONE;
@@ -2094,10 +2218,16 @@ struct FGameFeaturePluginState_Mounting : public FGameFeaturePluginState
 
 	virtual void EndState() override
 	{
+		TRACE_CPUPROFILER_EVENT_SCOPE(GFP_Mounting_EndState);
 		IInstallBundleManager::InstallBundleCompleteDelegate.RemoveAll(this);
-		FCoreDelegates::GetOnPakFileMounted2().RemoveAll(this);
+		if (PakFileMountedDelegateHandle.IsValid())
+		{
+			FCoreDelegates::GetOnPakFileMounted2().RemoveAll(this);
+			PakFileMountedDelegateHandle.Reset();
+		}
 	}
 };
+ENUM_CLASS_FLAGS(FGameFeaturePluginState_Mounting::ESubState);
 
 struct FWaitingForDependenciesTransitionPolicy
 {
@@ -3161,6 +3291,8 @@ void UGameFeaturePluginStateMachine::UpdateStateMachine()
 
 	auto DoCallbacks = [this](const UE::GameFeatures::FResult& Result, StateIt Begin, StateIt End)
 	{
+		TRACE_CPUPROFILER_EVENT_SCOPE(GFP_UpdateStateMachine_DoCallbacks);
+
 		for (StateIt iState = Begin; iState < End; ++iState)
 		{
 			if (FDestinationGameFeaturePluginState* DestState = AllStates[iState]->AsDestinationState())
@@ -3189,7 +3321,10 @@ void UGameFeaturePluginStateMachine::UpdateStateMachine()
 		bKeepProcessing = false;
 
 		FGameFeaturePluginStateStatus StateStatus;
-		AllStates[CurrentState]->UpdateState(StateStatus);
+		{
+			TRACE_CPUPROFILER_EVENT_SCOPE(GFP_UpdateStateMachine_UpdateState);
+			AllStates[CurrentState]->UpdateState(StateStatus);
+		}
 
 		if (StateStatus.TransitionToState == CurrentState)
 		{
@@ -3199,14 +3334,22 @@ void UGameFeaturePluginStateMachine::UpdateStateMachine()
 		if (StateStatus.TransitionToState != EGameFeaturePluginState::Uninitialized)
 		{
 			UE_LOG(LogGameFeatures, Verbose, TEXT("Game feature '%s' transitioning state (%s -> %s)"), *GetGameFeatureName(), *UE::GameFeatures::ToString(CurrentState), *UE::GameFeatures::ToString(StateStatus.TransitionToState));
-			AllStates[CurrentState]->EndState();
+			{
+				TRACE_CPUPROFILER_EVENT_SCOPE(GFP_UpdateStateMachine_EndState);
+				AllStates[CurrentState]->EndState();
+			}
 			CurrentStateInfo = FGameFeaturePluginStateInfo(StateStatus.TransitionToState);
 			CurrentState = StateStatus.TransitionToState;
 			check(CurrentState != EGameFeaturePluginState::MAX);
-			AllStates[CurrentState]->BeginState();
+			{
+				TRACE_CPUPROFILER_EVENT_SCOPE(GFP_UpdateStateMachine_BeginState);
+				AllStates[CurrentState]->BeginState();
+			}
 
 			if (CurrentState == EGameFeaturePluginState::Terminal)
 			{
+				TRACE_CPUPROFILER_EVENT_SCOPE(GFP_UpdateStateMachine_BeginTerm);
+
 				// Remove from gamefeature subsystem before calling back in case this GFP is reloaded on callback,
 				// but make sure we don't get destroyed from a GC during a callback
 				UGameFeaturesSubsystem::Get().BeginTermination(this);
@@ -3257,6 +3400,8 @@ void UGameFeaturePluginStateMachine::UpdateStateMachine()
 
 			if (CurrentState == EGameFeaturePluginState::Terminal)
 			{
+				TRACE_CPUPROFILER_EVENT_SCOPE(GFP_UpdateStateMachine_FinishTerm);
+
 				check(bKeepProcessing == false);
 				// Now that callbacks are done this machine can be cleaned up
 				UGameFeaturesSubsystem::Get().FinishTermination(this);
