Enable GameFeaturePlugin.AsyncLoad by default.

diff --git a/Engine/Plugins/Experimental/GameFeatures/Source/GameFeatures/Private/GameFeaturePluginStateMachine.cpp b/Engine/Plugins/Experimental/GameFeatures/Source/GameFeatures/Private/GameFeaturePluginStateMachine.cpp
--- a/Engine/Plugins/Experimental/GameFeatures/Source/GameFeatures/Private/GameFeaturePluginStateMachine.cpp
+++ b/Engine/Plugins/Experimental/GameFeatures/Source/GameFeatures/Private/GameFeaturePluginStateMachine.cpp
@@ -60,7 +60,7 @@ namespace UE::GameFeatures
 		ECVF_Default);
 
 	static TAutoConsoleVariable<bool> CVarAsyncLoad(TEXT("GameFeaturePlugin.AsyncLoad"),
-		false,
+		true,
 		TEXT("Enable to use aysnc loading"));
 
 	static TAutoConsoleVariable<bool> CVarAllowForceMonolithicShaderLibrary(TEXT("GameFeaturePlugin.AllowForceMonolithicShaderLibrary"),
