Fix missing platform-specific GFP config on Android

- It is possible for both a platform-extension dir and a platform-specific plugin config dir to exist, so testing directories is not good enough to determine the proper directory.
- Also account for empty ConfigPlatform in preview
- Ensure we don't try to load every single possible file if the base folders don't exist.

#rnx

diff --git a/Engine/Plugins/Experimental/GameFeatures/Source/GameFeatures/Private/GameFeatureData.cpp b/Engine/Plugins/Experimental/GameFeatures/Source/GameFeatures/Private/GameFeatureData.cpp
--- a/Engine/Plugins/Experimental/GameFeatures/Source/GameFeatures/Private/GameFeatureData.cpp
+++ b/Engine/Plugins/Experimental/GameFeatures/Source/GameFeatures/Private/GameFeatureData.cpp
@@ -121,11 +121,14 @@ void UGameFeatureData::InitializeHierarchicalPluginIniFiles(const FString& Plugi
 {
 	UDeviceProfileManager& DeviceProfileManager = UDeviceProfileManager::Get();
 
+	FString PlatformName = FPlatformProperties::IniPlatformName();
+
 #if ALLOW_OTHER_PLATFORM_CONFIG
 	const UDeviceProfile* PreviewDeviceProfile = DeviceProfileManager.GetPreviewDeviceProfile();
-	const FString PlatformName = PreviewDeviceProfile ? PreviewDeviceProfile->ConfigPlatform : FPlatformProperties::IniPlatformName();
-#else
-	const FString PlatformName = FPlatformProperties::IniPlatformName();
+	if (PreviewDeviceProfile)
+	{
+		PlatformName = PreviewDeviceProfile->ConfigPlatform.IsEmpty() ? PreviewDeviceProfile->DeviceType : PreviewDeviceProfile->ConfigPlatform;
+	}
 #endif
 
 	FString PluginInstalledStandardFilename = PluginInstalledFilename;
@@ -135,17 +138,17 @@ void UGameFeatureData::InitializeHierarchicalPluginIniFiles(const FString& Plugi
 	}
 
 	const FString PluginName = FPaths::GetBaseFilename(PluginInstalledStandardFilename);
-	const FString PluginConfigDir = FPaths::GetPath(PluginInstalledStandardFilename) / TEXT("Config/");
 	const FString PlatformExtensionDir = FPaths::ProjectPlatformExtensionsDir() / (PlatformName + "/");
 	const FString EngineConfigDir = FPaths::EngineConfigDir();
+	const FString PluginConfigDir = FPaths::GetPath(PluginInstalledStandardFilename) / TEXT("Config/");
+	const FString PluginPlatformConfigDir = FPaths::Combine(PluginConfigDir, PlatformName);
+	const FString PluginPlatformExtensionDir = FPaths::GetPath(PluginInstalledStandardFilename).Replace(*FPaths::ProjectDir(), *PlatformExtensionDir) / TEXT("Config");
 
-	// We'll look first in the game's platform extension dir for a plugin filesystem, and then default to the plugin folder
-	FString PluginPlatformConfigDir = FPaths::GetPath(PluginInstalledStandardFilename).Replace(*FPaths::ProjectDir(), *PlatformExtensionDir) / TEXT("Config");
-	if (!FPaths::DirectoryExists(PluginPlatformConfigDir))
+	// We're going to test a lot of paths, so only do it if some config actually exists
+	if (!FPaths::DirectoryExists(PluginPlatformConfigDir) && !FPaths::DirectoryExists(PluginPlatformExtensionDir))
 	{
-		PluginPlatformConfigDir = FPaths::Combine(PluginConfigDir, PlatformName);
+		return;
 	}
-	PluginPlatformConfigDir += TEXT("/");
 
 	const bool bIsBaseIniName = false;
 	const bool bForceReloadFromDisk = false;
@@ -362,7 +365,25 @@ void UGameFeatureData::InitializeHierarchicalPluginIniFiles(const FString& Plugi
 	for (const FIniLoadingParams& Ini : IniFilesToLoad)
 	{
 		const FString PluginIniName = Ini.bUsePlatformDir ? (PlatformName + PluginName + Ini.Name) : PluginName + Ini.Name;
-		const FString ConfigDirectory = Ini.bUsePlatformDir ? PluginPlatformConfigDir : PluginConfigDir;
+
+		FString ConfigDirectory;
+		if (Ini.bUsePlatformDir)
+		{
+			// We'll look first in the platform extension directory, then in the plugin's platform directory
+			if (FPaths::FileExists(FPaths::Combine(PluginPlatformExtensionDir, PluginIniName)))
+			{
+				ConfigDirectory = PluginPlatformExtensionDir;
+			}
+			else
+			{
+				ConfigDirectory = PluginPlatformConfigDir;
+			}
+		}
+		else
+		{
+			ConfigDirectory = PluginConfigDir;
+		}
+		ConfigDirectory += TEXT("/");
 
 		// @note: Loading the INI in this manner in order to have a record of relevant sections that were changed so that affected objects can be reloaded. By virtue of how
 		// this is parsed (standalone instead of being treated as a combined diff), the actual data within the sections will likely be incorrect. As an example, users adding
