Fix regression in GFP config where only plugins with platform-specific config get their config loaded

#rnx

diff --git a/Engine/Plugins/Experimental/GameFeatures/Source/GameFeatures/Private/GameFeatureData.cpp b/Engine/Plugins/Experimental/GameFeatures/Source/GameFeatures/Private/GameFeatureData.cpp
--- a/Engine/Plugins/Experimental/GameFeatures/Source/GameFeatures/Private/GameFeatureData.cpp
+++ b/Engine/Plugins/Experimental/GameFeatures/Source/GameFeatures/Private/GameFeatureData.cpp
@@ -145,7 +145,7 @@ void UGameFeatureData::InitializeHierarchicalPluginIniFiles(const FString& Plugi
 	const FString PluginPlatformExtensionDir = FPaths::GetPath(PluginInstalledStandardFilename).Replace(*FPaths::ProjectDir(), *PlatformExtensionDir) / TEXT("Config");
 
 	// We're going to test a lot of paths, so only do it if some config actually exists
-	if (!FPaths::DirectoryExists(PluginPlatformConfigDir) && !FPaths::DirectoryExists(PluginPlatformExtensionDir))
+	if (!FPaths::DirectoryExists(PluginConfigDir) && FPaths::DirectoryExists(PluginPlatformConfigDir) && !FPaths::DirectoryExists(PluginPlatformExtensionDir))
 	{
 		return;
 	}
