Turn on async GFP loading by default.
Don't allow async GFP loading during commandlets because they typically don't tick so we will never finish loading.

diff --git a/Engine/Plugins/Experimental/GameFeatures/Source/GameFeatures/Private/GameFeaturePluginStateMachine.cpp b/Engine/Plugins/Experimental/GameFeatures/Source/GameFeatures/Private/GameFeaturePluginStateMachine.cpp
--- a/Engine/Plugins/Experimental/GameFeatures/Source/GameFeatures/Private/GameFeaturePluginStateMachine.cpp
+++ b/Engine/Plugins/Experimental/GameFeatures/Source/GameFeatures/Private/GameFeaturePluginStateMachine.cpp
@@ -60,7 +60,7 @@ namespace UE::GameFeatures
 		ECVF_Default);
 
 	static TAutoConsoleVariable<bool> CVarAsyncLoad(TEXT("GameFeaturePlugin.AsyncLoad"),
-		false,
+		true,
 		TEXT("Enable to use aysnc loading"));
 
 	static TAutoConsoleVariable<bool> CVarAllowForceMonolithicShaderLibrary(TEXT("GameFeaturePlugin.AllowForceMonolithicShaderLibrary"),
@@ -584,9 +584,15 @@ bool FGameFeaturePluginState::AllowIniLoading() const
 	}
 }
 
+bool FGameFeaturePluginState::AllowAsyncLoading() const
+{
+	// Ticking is required for async loading
+	return !IsRunningCommandlet();
+}
+
 bool FGameFeaturePluginState::UseAsyncLoading() const
 {
-	return UE::GameFeatures::CVarAsyncLoad.GetValueOnGameThread();
+	return AllowAsyncLoading() && UE::GameFeatures::CVarAsyncLoad.GetValueOnGameThread();
 }
 
 /*
@@ -1504,6 +1510,7 @@ struct FGameFeaturePluginState_Downloading : public FGameFeaturePluginState
 		Cleanup();
 
 		check(StateProperties.GetPluginProtocol() == EGameFeaturePluginProtocol::InstallBundle);
+		checkf(AllowAsyncLoading(), TEXT("FGameFeaturePluginState::AllowAsyncLoading is while attempting to download GFP data."));
 
 		TSharedPtr<IInstallBundleManager> BundleManager = IInstallBundleManager::GetPlatformInstallBundleManager();
 		const TArray<FName>& InstallBundles = StateProperties.ProtocolMetadata.GetSubtype<FInstallBundlePluginProtocolMetaData>().InstallBundles;
diff --git a/Engine/Plugins/Experimental/GameFeatures/Source/GameFeatures/Private/GameFeaturePluginStateMachine.h b/Engine/Plugins/Experimental/GameFeatures/Source/GameFeatures/Private/GameFeaturePluginStateMachine.h
--- a/Engine/Plugins/Experimental/GameFeatures/Source/GameFeatures/Private/GameFeaturePluginStateMachine.h
+++ b/Engine/Plugins/Experimental/GameFeatures/Source/GameFeatures/Private/GameFeaturePluginStateMachine.h
@@ -398,6 +398,7 @@ protected:
 
 	bool AllowIniLoading() const;
 
+	bool AllowAsyncLoading() const;
 	bool UseAsyncLoading() const;
 
 private:
