Add a CVar to force GFP async loads, mostly usefull for ticked commandlets that want to opt-in.
Fix wrong compeletion flag being set on error

diff --git a/Engine/Plugins/Experimental/GameFeatures/Source/GameFeatures/Private/GameFeaturePluginStateMachine.cpp b/Engine/Plugins/Experimental/GameFeatures/Source/GameFeatures/Private/GameFeaturePluginStateMachine.cpp
--- a/Engine/Plugins/Experimental/GameFeatures/Source/GameFeatures/Private/GameFeaturePluginStateMachine.cpp
+++ b/Engine/Plugins/Experimental/GameFeatures/Source/GameFeatures/Private/GameFeaturePluginStateMachine.cpp
@@ -63,6 +63,10 @@ namespace UE::GameFeatures
 		true,
 		TEXT("Enable to use aysnc loading"));
 
+	static TAutoConsoleVariable<bool> CVarForceAsyncLoad(TEXT("GameFeaturePlugin.ForceAsyncLoad"),
+		false,
+		TEXT("Enable to force use of aysnc loading even if normally not allowed"));
+
 	static TAutoConsoleVariable<bool> CVarAllowForceMonolithicShaderLibrary(TEXT("GameFeaturePlugin.AllowForceMonolithicShaderLibrary"),
 		true,
 		TEXT("Enable to force only searching for monolithic shader libs when possible"));
@@ -592,7 +596,9 @@ bool FGameFeaturePluginState::AllowAsyncLoading() const
 
 bool FGameFeaturePluginState::UseAsyncLoading() const
 {
-	return AllowAsyncLoading() && UE::GameFeatures::CVarAsyncLoad.GetValueOnGameThread();
+	return 
+		(AllowAsyncLoading() && UE::GameFeatures::CVarAsyncLoad.GetValueOnGameThread()) ||
+		UE::GameFeatures::CVarForceAsyncLoad.GetValueOnGameThread();
 }
 
 /*
@@ -2122,7 +2128,7 @@ struct FGameFeaturePluginState_Mounting : public FGameFeaturePluginState
 
 		if (Result.HasError())
 		{
-			CompletedSubStates |= ESubState::MountPlugin;
+			CompletedSubStates |= ESubState::LoadAssetRegistry;
 			return;
 		}
 
