[Backout] - CL28230549
Original CL Desc
-----------------------------------------------------------------
Turn on async GFP loading by default.
Don't allow async GFP loading during commandlets because they typically don't tick so we will never finish loading.

diff --git a/Engine/Plugins/Experimental/GameFeatures/Source/GameFeatures/Private/GameFeaturePluginStateMachine.cpp b/Engine/Plugins/Experimental/GameFeatures/Source/GameFeatures/Private/GameFeaturePluginStateMachine.cpp
--- a/Engine/Plugins/Experimental/GameFeatures/Source/GameFeatures/Private/GameFeaturePluginStateMachine.cpp
+++ b/Engine/Plugins/Experimental/GameFeatures/Source/GameFeatures/Private/GameFeaturePluginStateMachine.cpp
@@ -60,12 +60,8 @@ namespace UE::GameFeatures
 		ECVF_Default);
 
 	static TAutoConsoleVariable<bool> CVarAsyncLoad(TEXT("GameFeaturePlugin.AsyncLoad"),
-		true,
-		TEXT("Enable to use aysnc loading"));
-
-	static TAutoConsoleVariable<bool> CVarForceAsyncLoad(TEXT("GameFeaturePlugin.ForceAsyncLoad"),
 		false,
-		TEXT("Enable to force use of aysnc loading even if normally not allowed"));
+		TEXT("Enable to use aysnc loading"));
 
 	static TAutoConsoleVariable<bool> CVarAllowForceMonolithicShaderLibrary(TEXT("GameFeaturePlugin.AllowForceMonolithicShaderLibrary"),
 		true,
@@ -588,17 +584,9 @@ bool FGameFeaturePluginState::AllowIniLoading() const
 	}
 }
 
-bool FGameFeaturePluginState::AllowAsyncLoading() const
-{
-	// Ticking is required for async loading
-	return !IsRunningCommandlet();
-}
-
 bool FGameFeaturePluginState::UseAsyncLoading() const
 {
-	return 
-		(AllowAsyncLoading() && UE::GameFeatures::CVarAsyncLoad.GetValueOnGameThread()) ||
-		UE::GameFeatures::CVarForceAsyncLoad.GetValueOnGameThread();
+	return UE::GameFeatures::CVarAsyncLoad.GetValueOnGameThread();
 }
 
 /*
@@ -1516,7 +1504,6 @@ struct FGameFeaturePluginState_Downloading : public FGameFeaturePluginState
 		Cleanup();
 
 		check(StateProperties.GetPluginProtocol() == EGameFeaturePluginProtocol::InstallBundle);
-		checkf(AllowAsyncLoading(), TEXT("FGameFeaturePluginState::AllowAsyncLoading is while attempting to download GFP data."));
 
 		TSharedPtr<IInstallBundleManager> BundleManager = IInstallBundleManager::GetPlatformInstallBundleManager();
 		const TArray<FName>& InstallBundles = StateProperties.ProtocolMetadata.GetSubtype<FInstallBundlePluginProtocolMetaData>().InstallBundles;
@@ -2128,7 +2115,7 @@ struct FGameFeaturePluginState_Mounting : public FGameFeaturePluginState
 
 		if (Result.HasError())
 		{
-			CompletedSubStates |= ESubState::LoadAssetRegistry;
+			CompletedSubStates |= ESubState::MountPlugin;
 			return;
 		}
 
diff --git a/Engine/Plugins/Experimental/GameFeatures/Source/GameFeatures/Private/GameFeaturePluginStateMachine.h b/Engine/Plugins/Experimental/GameFeatures/Source/GameFeatures/Private/GameFeaturePluginStateMachine.h
--- a/Engine/Plugins/Experimental/GameFeatures/Source/GameFeatures/Private/GameFeaturePluginStateMachine.h
+++ b/Engine/Plugins/Experimental/GameFeatures/Source/GameFeatures/Private/GameFeaturePluginStateMachine.h
@@ -398,7 +398,6 @@ protected:
 
 	bool AllowIniLoading() const;
 
-	bool AllowAsyncLoading() const;
 	bool UseAsyncLoading() const;
 
 private:
