Add new GameFeatureAction base for adding GameFeature related things to the audio engine

#tests PIE

diff --git a/Engine/Plugins/Experimental/GameFeatures/Source/GameFeatures/Private/GameFeatureAction_AudioActionBase.cpp b/Engine/Plugins/Experimental/GameFeatures/Source/GameFeatures/Private/GameFeatureAction_AudioActionBase.cpp
new file mode 100644
--- /dev/null
+++ b/Engine/Plugins/Experimental/GameFeatures/Source/GameFeatures/Private/GameFeatureAction_AudioActionBase.cpp
@@ -0,0 +1,56 @@
+// Copyright Epic Games, Inc. All Rights Reserved.
+
+#include "GameFeatureAction_AudioActionBase.h"
+#include "AudioDeviceManager.h"
+
+#include UE_INLINE_GENERATED_CPP_BY_NAME(GameFeatureAction_AudioActionBase)
+
+#define LOCTEXT_NAMESPACE "GameFeatures"
+
+void UGameFeatureAction_AudioActionBase::OnGameFeatureActivating(FGameFeatureActivatingContext& Context)
+{
+	DeviceCreatedHandle = FAudioDeviceManagerDelegates::OnAudioDeviceCreated.AddUObject(this, &UGameFeatureAction_AudioActionBase::OnDeviceCreated);
+	DeviceDestroyedHandle = FAudioDeviceManagerDelegates::OnAudioDeviceDestroyed.AddUObject(this, &UGameFeatureAction_AudioActionBase::OnDeviceDestroyed);
+
+	// Add to any existing devices
+	if (FAudioDeviceManager* AudioDeviceManager = FAudioDeviceManager::Get())
+	{
+		AudioDeviceManager->IterateOverAllDevices([this](Audio::FDeviceId DeviceId, FAudioDevice* InDevice)
+		{
+			AddToDevice(FAudioDeviceManager::Get()->GetAudioDevice(DeviceId));
+		});
+	}
+}
+
+void UGameFeatureAction_AudioActionBase::OnGameFeatureDeactivating(FGameFeatureDeactivatingContext& Context)
+{
+	// Remove from any existing devices
+	if (FAudioDeviceManager* AudioDeviceManager = FAudioDeviceManager::Get())
+	{
+		AudioDeviceManager->IterateOverAllDevices([this](Audio::FDeviceId DeviceId, FAudioDevice* InDevice)
+		{
+			RemoveFromDevice(FAudioDeviceManager::Get()->GetAudioDevice(DeviceId));
+		});
+	}
+
+	FAudioDeviceManagerDelegates::OnAudioDeviceCreated.Remove(DeviceCreatedHandle);
+	FAudioDeviceManagerDelegates::OnAudioDeviceDestroyed.Remove(DeviceDestroyedHandle);
+}
+
+void UGameFeatureAction_AudioActionBase::OnDeviceCreated(Audio::FDeviceId InDeviceId)
+{
+	if (FAudioDeviceManager* AudioDeviceManager = FAudioDeviceManager::Get())
+	{
+		AddToDevice(AudioDeviceManager->GetAudioDevice(InDeviceId));
+	}	
+}
+
+void UGameFeatureAction_AudioActionBase::OnDeviceDestroyed(Audio::FDeviceId InDeviceId)
+{
+	if (FAudioDeviceManager* AudioDeviceManager = FAudioDeviceManager::Get())
+	{
+		RemoveFromDevice(AudioDeviceManager->GetAudioDevice(InDeviceId));
+	}
+}
+
+#undef LOCTEXT_NAMESPACE
diff --git a/Engine/Plugins/Experimental/GameFeatures/Source/GameFeatures/Public/GameFeatureAction_AudioActionBase.h b/Engine/Plugins/Experimental/GameFeatures/Source/GameFeatures/Public/GameFeatureAction_AudioActionBase.h
new file mode 100644
--- /dev/null
+++ b/Engine/Plugins/Experimental/GameFeatures/Source/GameFeatures/Public/GameFeatureAction_AudioActionBase.h
@@ -0,0 +1,35 @@
+// Copyright Epic Games, Inc. All Rights Reserved.
+
+#pragma once
+
+#include "AudioDeviceHandle.h"
+#include "GameFeatureAction.h"
+#include "GameFeaturesSubsystem.h"
+
+#include "GameFeatureAction_AudioActionBase.generated.h"
+
+/**
+ * Base class for GameFeatureActions that affect the audio engine
+ */
+UCLASS(Abstract)
+class GAMEFEATURES_API UGameFeatureAction_AudioActionBase : public UGameFeatureAction
+{
+	GENERATED_BODY()
+
+public:
+	//~ Begin of UGameFeatureAction interface
+	virtual void OnGameFeatureActivating(FGameFeatureActivatingContext& Context) override;
+	virtual void OnGameFeatureDeactivating(FGameFeatureDeactivatingContext& Context) override;
+	//~ End of UGameFeatureAction interface
+
+protected:
+	/** Handle to delegate callbacks */
+	FDelegateHandle DeviceCreatedHandle;
+	FDelegateHandle DeviceDestroyedHandle;
+
+	void OnDeviceCreated(Audio::FDeviceId InDeviceId);
+	void OnDeviceDestroyed(Audio::FDeviceId InDeviceId);
+
+	virtual void AddToDevice(const FAudioDeviceHandle& AudioDeviceHandle) {}
+	virtual void RemoveFromDevice(const FAudioDeviceHandle& AudioDeviceHandle) {}
+};
