GFPs
Add bForceSyncLoading flag to FBuiltInGameFeaturePluginBehaviorOptions and FGameFeatureProtocolOptions.
Add ensures to check bForceSyncLoading for when waiting for async dependencies or downloading.  Synchronouse behavior can't be enforced for these scenarios.

diff --git a/Engine/Plugins/Experimental/GameFeatures/Source/GameFeatures/Private/GameFeaturePluginStateMachine.cpp b/Engine/Plugins/Experimental/GameFeatures/Source/GameFeatures/Private/GameFeaturePluginStateMachine.cpp
--- a/Engine/Plugins/Experimental/GameFeatures/Source/GameFeatures/Private/GameFeaturePluginStateMachine.cpp
+++ b/Engine/Plugins/Experimental/GameFeatures/Source/GameFeatures/Private/GameFeaturePluginStateMachine.cpp
@@ -590,15 +590,12 @@ bool FGameFeaturePluginState::AllowIniLoading() const
 
 bool FGameFeaturePluginState::AllowAsyncLoading() const
 {
-	// Ticking is required for async loading
-	return !IsRunningCommandlet();
+	return StateProperties.AllowAsyncLoading();
 }
 
 bool FGameFeaturePluginState::UseAsyncLoading() const
 {
-	return
-		(AllowAsyncLoading() && UE::GameFeatures::CVarAsyncLoad.GetValueOnGameThread()) ||
-		UE::GameFeatures::CVarForceAsyncLoad.GetValueOnGameThread();
+	return AllowAsyncLoading() && UE::GameFeatures::CVarAsyncLoad.GetValueOnGameThread();
 }
 
 /*
@@ -662,9 +659,15 @@ struct FTransitionDependenciesGameFeaturePluginState : public FGameFeaturePlugin
 
 			UE_CLOG(Dependencies.Num() > 0, LogGameFeatures, Verbose, TEXT("Found %i dependencies for %s"), Dependencies.Num(), *StateProperties.PluginName);
 
+			const bool bAllowAsyncLoading = AllowAsyncLoading();
+
 			RemainingDependencies.Reserve(Dependencies.Num());
 			for (UGameFeaturePluginStateMachine* Dependency : Dependencies)
 			{
+				ensureMsgf(bAllowAsyncLoading || !Dependency->AllowAsyncLoading(), 
+					TEXT("FGameFeaturePluginState::AllowAsyncLoading is false for %s but true for dependency being waited on %s"), 
+					*StateProperties.PluginName, *Dependency->GetPluginURL());
+
 				RemainingDependencies.Emplace(Dependency, MakeValue());
 				TransitionDependency(Dependency);
 			}
@@ -1516,7 +1519,7 @@ struct FGameFeaturePluginState_Downloading : public FGameFeaturePluginState
 		Cleanup();
 
 		check(StateProperties.GetPluginProtocol() == EGameFeaturePluginProtocol::InstallBundle);
-		ensureMsgf(AllowAsyncLoading(), TEXT("FGameFeaturePluginState::AllowAsyncLoading is while attempting to download GFP data."));
+		ensureMsgf(AllowAsyncLoading(), TEXT("FGameFeaturePluginState::AllowAsyncLoading is false while attempting to download GFP data."));
 
 		TSharedPtr<IInstallBundleManager> BundleManager = IInstallBundleManager::GetPlatformInstallBundleManager();
 		const TArray<FName>& InstallBundles = StateProperties.ProtocolMetadata.GetSubtype<FInstallBundlePluginProtocolMetaData>().InstallBundles;
@@ -3284,6 +3287,11 @@ bool UGameFeaturePluginStateMachine::IsAvailable() const
 	return GetCurrentState() >= EGameFeaturePluginState::StatusKnown;
 }
 
+bool UGameFeaturePluginStateMachine::AllowAsyncLoading() const
+{
+	return StateProperties.AllowAsyncLoading();
+}
+
 UGameFeatureData* UGameFeaturePluginStateMachine::GetGameFeatureDataForActivePlugin()
 {
 	if (GetCurrentState() == EGameFeaturePluginState::Active)
@@ -3691,3 +3699,12 @@ FGameFeatureProtocolOptions FGameFeaturePluginStateMachineProperties::RecyclePro
 	return Result;
 }
 
+bool FGameFeaturePluginStateMachineProperties::AllowAsyncLoading() const
+{
+	// Ticking is required for async loading
+	// The local bForceSyncLoading should take precedence over UE::GameFeatures::CVarForceAsyncLoad
+	return
+		!ProtocolOptions.bForceSyncLoading &&
+		(!IsRunningCommandlet() || UE::GameFeatures::CVarForceAsyncLoad.GetValueOnGameThread());
+}
+
diff --git a/Engine/Plugins/Experimental/GameFeatures/Source/GameFeatures/Private/GameFeaturePluginStateMachine.h b/Engine/Plugins/Experimental/GameFeatures/Source/GameFeatures/Private/GameFeaturePluginStateMachine.h
--- a/Engine/Plugins/Experimental/GameFeatures/Source/GameFeatures/Private/GameFeaturePluginStateMachine.h
+++ b/Engine/Plugins/Experimental/GameFeatures/Source/GameFeatures/Private/GameFeaturePluginStateMachine.h
@@ -319,6 +319,9 @@ struct FGameFeaturePluginStateMachineProperties
 
 	/** Returns protocol options suitable for reuse by another state machine */
 	FGameFeatureProtocolOptions RecycleProtocolOptions() const;
+
+	/** Whether this machine is allowed to be asynchronous */
+	bool AllowAsyncLoading() const;
 };
 
 /** Input and output information for a state's UpdateState */
@@ -517,6 +520,9 @@ public:
 	/** Returns true if the plugin is available to download/load. Only call if IsStatusKnown is true */
 	bool IsAvailable() const;
 
+	/** Whether this machine is allowed to be asynchronous */
+	bool AllowAsyncLoading() const;
+
 	/** If the plugin is activated already, we will retrieve its game feature data */
 	UGameFeatureData* GetGameFeatureDataForActivePlugin();
 
diff --git a/Engine/Plugins/Experimental/GameFeatures/Source/GameFeatures/Private/GameFeaturesSubsystem.cpp b/Engine/Plugins/Experimental/GameFeatures/Source/GameFeatures/Private/GameFeaturesSubsystem.cpp
--- a/Engine/Plugins/Experimental/GameFeatures/Source/GameFeatures/Private/GameFeaturesSubsystem.cpp
+++ b/Engine/Plugins/Experimental/GameFeatures/Source/GameFeatures/Private/GameFeaturesSubsystem.cpp
@@ -1462,7 +1462,9 @@ void UGameFeaturesSubsystem::LoadBuiltInGameFeaturePlugin(const TSharedRef<IPlug
 			const bool bShouldProcess = AdditionalFilter(Plugin->GetDescriptorFileName(), PluginDetails, BehaviorOptions);
 			if (bShouldProcess)
 			{
-				UGameFeaturePluginStateMachine* StateMachine = FindOrCreateGameFeaturePluginStateMachine(PluginURL, FGameFeatureProtocolOptions());
+				FGameFeatureProtocolOptions ProtocolOptions;
+				ProtocolOptions.bForceSyncLoading = BehaviorOptions.bForceSyncLoading;
+				UGameFeaturePluginStateMachine* StateMachine = FindOrCreateGameFeaturePluginStateMachine(PluginURL, ProtocolOptions);
 
 				EBuiltInAutoState InitialAutoState = (BehaviorOptions.AutoStateOverride != EBuiltInAutoState::Invalid) ? 
 					BehaviorOptions.AutoStateOverride : PluginDetails.BuiltInAutoState;
@@ -1475,7 +1477,7 @@ void UGameFeaturesSubsystem::LoadBuiltInGameFeaturePlugin(const TSharedRef<IPlug
 
 				// If we're already at the destination or beyond, don't transition back
 				FGameFeaturePluginStateRange Destination(DestinationState, EGameFeaturePluginState::Active);
-				ChangeGameFeatureDestination(StateMachine, Destination, 
+				ChangeGameFeatureDestination(StateMachine, ProtocolOptions, Destination,
 					FGameFeaturePluginChangeStateComplete::CreateWeakLambda(this, [this, StateMachine, Destination, CompleteDelegate](const UE::GameFeatures::FResult& Result)
 					{
 						LoadBuiltInGameFeaturePluginComplete(Result, StateMachine, Destination);
@@ -2130,6 +2132,11 @@ bool UGameFeaturesSubsystem::FindOrCreatePluginDependencyStateMachines(const FSt
 			{
 				DepProtocolOptions = InDepProtocolOptions;
 			}
+			else
+			{
+				// Always propogate non-protocol specific flags
+				DepProtocolOptions.bForceSyncLoading = InDepProtocolOptions.bForceSyncLoading;
+			}
 
 			UGameFeaturePluginStateMachine* ResolvedDependency = FindOrCreateGameFeaturePluginStateMachine(DependencyURL, DepProtocolOptions);
 			check(ResolvedDependency);
diff --git a/Engine/Plugins/Experimental/GameFeatures/Source/GameFeatures/Public/GameFeaturesSubsystem.h b/Engine/Plugins/Experimental/GameFeatures/Source/GameFeatures/Public/GameFeaturesSubsystem.h
--- a/Engine/Plugins/Experimental/GameFeatures/Source/GameFeatures/Public/GameFeaturesSubsystem.h
+++ b/Engine/Plugins/Experimental/GameFeatures/Source/GameFeatures/Public/GameFeaturesSubsystem.h
@@ -256,6 +256,7 @@ struct FGameFeaturePluginDetails
 struct FBuiltInGameFeaturePluginBehaviorOptions
 {
 	EBuiltInAutoState AutoStateOverride = EBuiltInAutoState::Invalid;
+	bool bForceSyncLoading = false;
 };
 
 /** Struct used to transform a GameFeaturePlugin URL into something that can uniquely identify the GameFeaturePlugin
@@ -351,8 +352,11 @@ struct GAMEFEATURES_API FInstallBundlePluginProtocolOptions
 struct FGameFeatureProtocolOptions : public TUnion<FInstallBundlePluginProtocolOptions, FNull>
 {
 	FGameFeatureProtocolOptions() { SetSubtype<FNull>(); }
-	FGameFeatureProtocolOptions(const FInstallBundlePluginProtocolOptions& InOptions) : TUnion(InOptions) {}
-	FGameFeatureProtocolOptions(FNull InOptions) { SetSubtype<FNull>(InOptions); }
+	explicit FGameFeatureProtocolOptions(const FInstallBundlePluginProtocolOptions& InOptions) : TUnion(InOptions) {}
+	explicit FGameFeatureProtocolOptions(FNull InOptions) { SetSubtype<FNull>(InOptions); }
+
+	/** Force this GFP to load synchronously even if async loading is allowed */
+	bool bForceSyncLoading = false;
 };
 
 /** The manager subsystem for game features */
