Fix regression where device profile preview no longer works with GFP device profiles

#rnx

diff --git a/Engine/Plugins/Experimental/GameFeatures/Source/GameFeatures/Private/GameFeatureData.cpp b/Engine/Plugins/Experimental/GameFeatures/Source/GameFeatures/Private/GameFeatureData.cpp
--- a/Engine/Plugins/Experimental/GameFeatures/Source/GameFeatures/Private/GameFeatureData.cpp
+++ b/Engine/Plugins/Experimental/GameFeatures/Source/GameFeatures/Private/GameFeatureData.cpp
@@ -370,7 +370,7 @@ void UGameFeatureData::InitializeHierarchicalPluginIniFiles(const FString& Plugi
 		if (Ini.bUsePlatformDir)
 		{
 			// We'll look first in the platform extension directory, then in the plugin's platform directory
-			if (FPaths::FileExists(FPaths::Combine(PluginPlatformExtensionDir, PluginIniName)))
+			if (FPaths::FileExists(FPaths::Combine(PluginPlatformExtensionDir, PluginIniName + ".ini")))
 			{
 				ConfigDirectory = PluginPlatformExtensionDir;
 			}
