Fix missing Android fragment matching rules when using runtime device profiles on mobile

#rnx

diff --git a/Engine/Plugins/Experimental/GameFeatures/Source/GameFeatures/Private/GameFeatureData.cpp b/Engine/Plugins/Experimental/GameFeatures/Source/GameFeatures/Private/GameFeatureData.cpp
--- a/Engine/Plugins/Experimental/GameFeatures/Source/GameFeatures/Private/GameFeatureData.cpp
+++ b/Engine/Plugins/Experimental/GameFeatures/Source/GameFeatures/Private/GameFeatureData.cpp
@@ -265,6 +265,28 @@ void UGameFeatureData::InitializeHierarchicalPluginIniFiles(const FString& Plugi
 								RuntimeProfile.Add("DeviceType", PlatformName);
 								RuntimeProfile.Add("BaseProfileName", FConfigValue(Profile->GetName()));
 
+								// Inject the parent's matched fragments into the config, if any
+								if (Profile->GetName().Contains("MatchedFragments"))
+								{
+									FString MatchingRulesSectionName = Profile->GetName() + TEXT(" ") + UDeviceProfile::StaticClass()->GetName();
+									FString MatchingRulesArrayName = TEXT("MatchingRules");
+									TArray<FString> MatchingRulesArray;
+
+#if ALLOW_OTHER_PLATFORM_CONFIG
+									FConfigCacheIni* PlatformConfigSystem = FConfigCacheIni::ForPlatform(*PlatformName);
+#else
+									FConfigCacheIni* PlatformConfigSystem = GConfig;
+#endif
+
+									PlatformConfigSystem->GetArray(*MatchingRulesSectionName, *MatchingRulesArrayName, MatchingRulesArray, GDeviceProfilesIni);
+									UE_LOG(LogGameFeatures, Display, TEXT("Found %d fragment matching rules"), MatchingRulesArray.Num());
+
+									for (const FString& Rule : MatchingRulesArray)
+									{
+										RuntimeProfile.Add("+MatchingRules", FConfigValue(Rule));
+									}
+								}
+
 								// Add fragment includes
 								for (const auto& FragmentInclude : FragmentIncludes)
 								{
