Added support for preloading data registries and sources in editor commandlets for automation of editor operations that rely on bPreloadInEditor

diff --git a/Engine/Plugins/Experimental/GameFeatures/Source/GameFeatures/Private/GameFeatureAction_DataRegistry.cpp b/Engine/Plugins/Experimental/GameFeatures/Source/GameFeatures/Private/GameFeatureAction_DataRegistry.cpp
--- a/Engine/Plugins/Experimental/GameFeatures/Source/GameFeatures/Private/GameFeatureAction_DataRegistry.cpp
+++ b/Engine/Plugins/Experimental/GameFeatures/Source/GameFeatures/Private/GameFeatureAction_DataRegistry.cpp
@@ -139,8 +139,14 @@ void UGameFeatureAction_DataRegistry::OnGameFeatureDeactivating(FGameFeatureDeac
 
 bool UGameFeatureAction_DataRegistry::ShouldPreloadAtRegistration()
 {
-	// We want to preload in interactive editor sessions only
-	return (GIsEditor && !IsRunningCommandlet() && bPreloadInEditor);
+	if (IsRunningCommandlet())
+	{
+		return bPreloadInCommandlets;
+	}
+	else
+	{
+		return GIsEditor && bPreloadInEditor;
+	}
 }
 
 #if WITH_EDITORONLY_DATA
diff --git a/Engine/Plugins/Experimental/GameFeatures/Source/GameFeatures/Public/GameFeatureAction_DataRegistry.h b/Engine/Plugins/Experimental/GameFeatures/Source/GameFeatures/Public/GameFeatureAction_DataRegistry.h
--- a/Engine/Plugins/Experimental/GameFeatures/Source/GameFeatures/Public/GameFeatureAction_DataRegistry.h
+++ b/Engine/Plugins/Experimental/GameFeatures/Source/GameFeatures/Public/GameFeatureAction_DataRegistry.h
@@ -40,6 +40,10 @@ private:
 	/** If true, this will preload the registries when the feature is registered in the editor to support the editor pickers */
 	UPROPERTY(EditAnywhere, Category = "Registry Data")
 	bool bPreloadInEditor;
+
+	/** If true, this will preload the registries when the feature is registered whilst a commandlet is running */
+	UPROPERTY(EditAnywhere, Category = "Registry Data")
+	bool bPreloadInCommandlets;
 };
 
 #if UE_ENABLE_INCLUDE_ORDER_DEPRECATED_IN_5_2
