GetGameFeaturePluginDetails API cleanup
- Rename the private version of GetGameFeaturePluginDetails to GetGameFeaturePluginDetailsInternal
- Deprecate the public version of GetGameFeaturePluginDetails taking an IPlugin, and rename it to GetBuiltInGameFeaturePluginDetails
- Make a new public version of GetGameFeaturePluginDetails that takes a GFP URL

The new public version of GetGameFeaturePluginDetails() should work on any GFP.  The old version was generally only useful for built-in plugins that were already known to the Engine.
This change allows game code to first, install on demand GFPs, and then Query their BuiltInAutoState via GetGameFeaturePluginDetails() after the uplugin is installed and  transition the GFP to that state.
This requires more work from game code than for built-in GFPs but allows for re-use of BuiltInAutoState and easier migration of built-in GFPs to on demand.

diff --git a/Engine/Plugins/Experimental/GameFeatures/Source/GameFeatures/Private/GameFeaturesSubsystem.cpp b/Engine/Plugins/Experimental/GameFeatures/Source/GameFeatures/Private/GameFeaturesSubsystem.cpp
--- a/Engine/Plugins/Experimental/GameFeatures/Source/GameFeatures/Private/GameFeaturesSubsystem.cpp
+++ b/Engine/Plugins/Experimental/GameFeatures/Source/GameFeatures/Private/GameFeaturesSubsystem.cpp
@@ -1454,7 +1454,7 @@ void UGameFeaturesSubsystem::LoadBuiltInGameFeaturePlugin(const TSharedRef<IPlug
 
 	FString PluginURL;
 	FGameFeaturePluginDetails PluginDetails;
-	if (GetGameFeaturePluginDetails(Plugin, PluginURL, PluginDetails))
+	if (GetBuiltInGameFeaturePluginDetails(Plugin, PluginURL, PluginDetails))
 	{
 		if (GameSpecificPolicies->IsPluginAllowed(PluginURL))
 		{
@@ -1729,6 +1729,11 @@ EGameFeaturePluginState UGameFeaturesSubsystem::GetPluginState(FGameFeaturePlugi
 }
 
 bool UGameFeaturesSubsystem::GetGameFeaturePluginDetails(const TSharedRef<IPlugin>& Plugin, FString& OutPluginURL, FGameFeaturePluginDetails& OutPluginDetails) const
+{
+	return GetBuiltInGameFeaturePluginDetails(Plugin, OutPluginURL, OutPluginDetails);
+}
+
+bool UGameFeaturesSubsystem::GetBuiltInGameFeaturePluginDetails(const TSharedRef<IPlugin>& Plugin, FString& OutPluginURL, FGameFeaturePluginDetails& OutPluginDetails) const
 {
 	// @TODO: this problematic because it assumes file protocol.
 	// Ideally this would work with any protocol, but for current uses cases the exact protocol doesn't seem to matter.
@@ -1750,14 +1755,20 @@ bool UGameFeaturesSubsystem::GetGameFeaturePluginDetails(const TSharedRef<IPlugi
 
 		if (bIsFileProtocol)
 		{
-			return GetGameFeaturePluginDetails(OutPluginURL, PluginDescriptorFilename, OutPluginDetails);
+			return GetGameFeaturePluginDetailsInternal(PluginDescriptorFilename, OutPluginDetails);
 		}
 	}
 
 	return false;
 }
 
-bool UGameFeaturesSubsystem::GetGameFeaturePluginDetails(const FString& PluginURL, const FString& PluginDescriptorFilename, FGameFeaturePluginDetails& OutPluginDetails) const
+bool UGameFeaturesSubsystem::GetGameFeaturePluginDetails(FString PluginURL, FGameFeaturePluginDetails& OutPluginDetails) const
+{
+	FGameFeaturePluginIdentifier Ident(MoveTemp(PluginURL));
+	return GetGameFeaturePluginDetailsInternal(FString(Ident.GetIdentifyingString()), OutPluginDetails);
+}
+
+bool UGameFeaturesSubsystem::GetGameFeaturePluginDetailsInternal(const FString& PluginDescriptorFilename, FGameFeaturePluginDetails& OutPluginDetails) const
 {
 	TRACE_CPUPROFILER_EVENT_SCOPE(GFP_GetPluginDetails);
 
@@ -2101,7 +2112,7 @@ void UGameFeaturesSubsystem::FinishTermination(UGameFeaturePluginStateMachine* M
 bool UGameFeaturesSubsystem::FindOrCreatePluginDependencyStateMachines(const FString& PluginURL, const FString& PluginFilename, const FGameFeatureProtocolOptions& InDepProtocolOptions, TArray<UGameFeaturePluginStateMachine*>& OutDependencyMachines)
 {
 	FGameFeaturePluginDetails Details;
-	if (GetGameFeaturePluginDetails(PluginURL, PluginFilename, Details))
+	if (GetGameFeaturePluginDetailsInternal(PluginFilename, Details))
 	{
 		for (const FGameFeaturePluginReferenceDetails& PluginDependency : Details.PluginDependencies)
 		{
@@ -2152,7 +2163,7 @@ bool UGameFeaturesSubsystem::FindOrCreatePluginDependencyStateMachines(const FSt
 bool UGameFeaturesSubsystem::FindPluginDependencyStateMachinesToActivate(const FString& PluginURL, const FString& PluginFilename, TArray<UGameFeaturePluginStateMachine*>& OutDependencyMachines) const
 {
 	FGameFeaturePluginDetails Details;
-	if (GetGameFeaturePluginDetails(PluginURL, PluginFilename, Details))
+	if (GetGameFeaturePluginDetailsInternal(PluginFilename, Details))
 	{
 		for (const FGameFeaturePluginReferenceDetails& PluginDependency : Details.PluginDependencies)
 		{
diff --git a/Engine/Plugins/Experimental/GameFeatures/Source/GameFeatures/Public/GameFeaturesSubsystem.h b/Engine/Plugins/Experimental/GameFeatures/Source/GameFeatures/Public/GameFeaturesSubsystem.h
--- a/Engine/Plugins/Experimental/GameFeatures/Source/GameFeatures/Public/GameFeaturesSubsystem.h
+++ b/Engine/Plugins/Experimental/GameFeatures/Source/GameFeatures/Public/GameFeaturesSubsystem.h
@@ -560,8 +560,15 @@ public:
 	EGameFeaturePluginState GetPluginState(FGameFeaturePluginIdentifier PluginIdentifier) const;
 
 	/** Gets relevant properties out of a uplugin file */
+	UE_DEPRECATED(5.4, "Use GetBuiltInGameFeaturePluginDetails instead")
 	bool GetGameFeaturePluginDetails(const TSharedRef<IPlugin>& Plugin, FString& OutPluginURL, struct FGameFeaturePluginDetails& OutPluginDetails) const;
 
+	/** Gets relevant properties out of a uplugin file. Should only be used for built-in GFPs */
+	bool GetBuiltInGameFeaturePluginDetails(const TSharedRef<IPlugin>& Plugin, FString& OutPluginURL, struct FGameFeaturePluginDetails& OutPluginDetails) const;
+
+	/** Gets relevant properties out of a uplugin file if it's installed */
+	bool GetGameFeaturePluginDetails(FString PluginURL, struct FGameFeaturePluginDetails& OutPluginDetails) const;
+
 	/** Determine the initial feature state for a built-in plugin */
 	static EBuiltInAutoState DetermineBuiltInInitialFeatureState(TSharedPtr<FJsonObject> Descriptor, const FString& ErrorContext);
 
@@ -620,7 +627,7 @@ private:
 	const UGameFeatureData* GetRegisteredDataForStateMachine(UGameFeaturePluginStateMachine* GFSM) const;
 
 	/** Gets relevant properties out of a uplugin file */
-	bool GetGameFeaturePluginDetails(const FString& PluginURL, const FString& PluginDescriptorFilename, struct FGameFeaturePluginDetails& OutPluginDetails) const;
+	bool GetGameFeaturePluginDetailsInternal(const FString& PluginDescriptorFilename, struct FGameFeaturePluginDetails& OutPluginDetails) const;
 
 	/** Prunes any cached GFP details */
 	void PruneCachedGameFeaturePluginDetails(const FString& PluginURL, const FString& PluginDescriptorFilename) const;
diff --git a/Engine/Plugins/Experimental/PluginAudit/Source/Private/SPluginAuditBrowser.cpp b/Engine/Plugins/Experimental/PluginAudit/Source/Private/SPluginAuditBrowser.cpp
--- a/Engine/Plugins/Experimental/PluginAudit/Source/Private/SPluginAuditBrowser.cpp
+++ b/Engine/Plugins/Experimental/PluginAudit/Source/Private/SPluginAuditBrowser.cpp
@@ -161,7 +161,7 @@ void SPluginAuditBrowser::BuildPluginList()
 	{
 		FString PluginURL;
 		FGameFeaturePluginDetails PluginDetails;
-		if (GameFeaturesSubsystem.GetGameFeaturePluginDetails(Plugin, PluginURL, PluginDetails))
+		if (GameFeaturesSubsystem.GetBuiltInGameFeaturePluginDetails(Plugin, PluginURL, PluginDetails))
 		{
 			if (Policy.WillPluginBeCooked(Plugin->GetDescriptorFileName(), PluginDetails))
 			{
