Mirroring the activate dependencies pattern to selectively deactivate listed plugin dependencies when a plugin is deactivated.

New States
- Added DeactivatingDependencies state between Active and Deactivating that will selectively deactivate required dependencies and wait for completetion.
- Added ErrorDeactivatingDependencies to handle errors in the DeactivatingDependencies state. Can only transtion on to DeactivatingDependencies.

diff --git a/Engine/Plugins/Experimental/GameFeatures/Source/GameFeatures/Private/GameFeaturePluginStateMachine.cpp b/Engine/Plugins/Experimental/GameFeatures/Source/GameFeatures/Private/GameFeaturePluginStateMachine.cpp
--- a/Engine/Plugins/Experimental/GameFeatures/Source/GameFeatures/Private/GameFeaturePluginStateMachine.cpp
+++ b/Engine/Plugins/Experimental/GameFeatures/Source/GameFeatures/Private/GameFeaturePluginStateMachine.cpp
@@ -2794,6 +2794,57 @@ struct FGameFeaturePluginState_Loaded : public FDestinationGameFeaturePluginStat
 	}
 };
 
+struct FGameFeaturePluginState_ErrorDeactivatingDependencies : public FErrorGameFeaturePluginState
+{
+	FGameFeaturePluginState_ErrorDeactivatingDependencies(FGameFeaturePluginStateMachineProperties& InStateProperties) : FErrorGameFeaturePluginState(InStateProperties) {}
+
+	virtual void UpdateState(FGameFeaturePluginStateStatus& StateStatus) override
+	{
+		if (StateProperties.Destination < EGameFeaturePluginState::ErrorDeactivatingDependencies)
+		{
+			StateStatus.SetTransition(EGameFeaturePluginState::DeactivatingDependencies);
+		}
+		else if (StateProperties.Destination > EGameFeaturePluginState::ErrorDeactivatingDependencies)
+		{
+			StateStatus.SetTransition(EGameFeaturePluginState::DeactivatingDependencies);
+		}
+	}
+};
+
+struct FDeactivatingDependenciesTransitionPolicy
+{
+	static bool GetPluginDependencyStateMachines(const FGameFeaturePluginStateMachineProperties& InStateProperties, TArray<UGameFeaturePluginStateMachine*>& OutDependencyMachines)
+	{
+		UGameFeaturesSubsystem& GameFeaturesSubsystem = UGameFeaturesSubsystem::Get();
+
+		return GameFeaturesSubsystem.FindPluginDependencyStateMachinesToDeactivate(
+			*InStateProperties.PluginIdentifier.GetFullPluginURL(), InStateProperties.PluginInstalledFilename, OutDependencyMachines);
+	}
+
+	static FGameFeaturePluginStateRange GetDependencyStateRange()
+	{
+		return FGameFeaturePluginStateRange(EGameFeaturePluginState::Terminal, EGameFeaturePluginState::Loaded);
+	}
+
+	static EGameFeaturePluginState GetTransitionState()
+	{
+		return EGameFeaturePluginState::Deactivating;
+	}
+
+	static EGameFeaturePluginState GetErrorState()
+	{
+		return EGameFeaturePluginState::ErrorDeactivatingDependencies;
+	}
+};
+
+struct FGameFeaturePluginState_DeactivatingDependencies : public FTransitionDependenciesGameFeaturePluginState<FDeactivatingDependenciesTransitionPolicy>
+{
+	FGameFeaturePluginState_DeactivatingDependencies(FGameFeaturePluginStateMachineProperties& InStateProperties)
+		: FTransitionDependenciesGameFeaturePluginState(InStateProperties)
+	{
+	}
+};
+
 struct FGameFeaturePluginState_Deactivating : public FGameFeaturePluginState
 {
 	FGameFeaturePluginState_Deactivating(FGameFeaturePluginStateMachineProperties& InStateProperties) : FGameFeaturePluginState(InStateProperties) {}
@@ -2978,7 +3029,7 @@ struct FGameFeaturePluginState_Active : public FDestinationGameFeaturePluginStat
 	{
 		if (StateProperties.Destination < EGameFeaturePluginState::Active)
 		{
-			StateStatus.SetTransition(EGameFeaturePluginState::Deactivating);
+			StateStatus.SetTransition(EGameFeaturePluginState::DeactivatingDependencies);
 		}
 	}
 };
diff --git a/Engine/Plugins/Experimental/GameFeatures/Source/GameFeatures/Private/GameFeaturePluginStateMachine.h b/Engine/Plugins/Experimental/GameFeatures/Source/GameFeatures/Private/GameFeaturePluginStateMachine.h
--- a/Engine/Plugins/Experimental/GameFeatures/Source/GameFeatures/Private/GameFeaturePluginStateMachine.h
+++ b/Engine/Plugins/Experimental/GameFeatures/Source/GameFeatures/Private/GameFeaturePluginStateMachine.h
@@ -120,37 +120,43 @@ Transition states are expected to transition the machine to another state after
                |            |
                +-^--------+-+
                  |        |
-           ------~--------~-------------------------------------
-           |     |        |                             ^      |
-        +--v-----+--+  +--v--------+                    |    +-+------------+
-        |           |  |           |                    |    |      !       |
-        | Unloading |  |  Loading  <--------------------~----> ErrorLoading |
-        |           |  |           |                    |    |              |
-        +--------^--+  +--+--------+                    |    +--------------+
-                 |        |                             |
-               +-+--------v-+                           |
-               |      *     |                           |
-               |   Loaded   |                           |
-               |            |                           |
-               +-^--------+-+                           |
-                 |        |                             |
-        +--------+---+  +-v----------------------+   +--+--------------------------+
-        |            |  |                        |   |             !               |
-        |Deactivating|  | ActivatingDependencies <---> ErrorActivatingDependencies |
-        |            |  |                        |   |                             |
-        +--------^---+  +-+----------------------+   +-----------------------------+
-                 |        |
-				 |      +-v----------+
-				 |      |            |
-				 |      | Activating |
-				 |      |            |
-				 |      +-+----------+
-				 |        |
-               +-+--------v-+
-               |      *     |
-               |   Active   |
-               |            |
-               +------------+
+           ------~--------~---------------------------------------
+           |     |        |                               ^      |
+        +--v-----+--+  +--v--------+                      |    +-+------------+
+        |           |  |           |                      |    |      !       |
+        | Unloading |  |  Loading  <-------------------- ~----> ErrorLoading |
+        |           |  |           |                      |    |              |
+        +--------^--+  +--+--------+                      |    +--------------+
+                 |        |                               |
+               +-+--------v-+                             |
+               |      *     |                             |
+               |   Loaded   |                             |
+               |            |                             |
+               +-^--------+-+                             |
+		         |        |                               |
+        +--------+---+  +-v------------------------+   +--+--------------------------+
+        |            |  |                          |   |             !               |
+        |Deactivating|  |  ActivatingDependencies  <---> ErrorActivatingDependencies |
+        |            |  |                          |   |                             |
+        +-^----------+  +---------------------+----+   +-----------------------------+
+          |                                   |
+		  |  +-----------------------------+  |
+		  |  |              !              |  |
+		  |  |ErrorDeactivatingDependencies|  |
+		  |  |                             |  |
+		  |  +--^--------------------------+  |
+		  |     |                             |
+		+-+-----v----------------+          +-v----------+
+		|		                 |          |            |
+		|DeactivatingDependencies|          | Activating |
+		|	   	                 |          |            |
+		+----------------------^-+          +---+--------+
+		                       |                |
+		     	             +-+----------------v-+
+                             |          *         |
+                             |       Active       |
+                             |                    |
+                             +--------------------+
 */
 
 namespace UE::GameFeatures
diff --git a/Engine/Plugins/Experimental/GameFeatures/Source/GameFeatures/Private/GameFeaturesSubsystem.cpp b/Engine/Plugins/Experimental/GameFeatures/Source/GameFeatures/Private/GameFeaturesSubsystem.cpp
--- a/Engine/Plugins/Experimental/GameFeatures/Source/GameFeatures/Private/GameFeaturesSubsystem.cpp
+++ b/Engine/Plugins/Experimental/GameFeatures/Source/GameFeatures/Private/GameFeaturesSubsystem.cpp
@@ -2205,6 +2205,50 @@ bool UGameFeaturesSubsystem::FindPluginDependencyStateMachinesToActivate(const F
 	return false;
 }
 
+bool UGameFeaturesSubsystem::FindPluginDependencyStateMachinesToDeactivate(const FString& PluginURL, const FString& PluginFilename, TArray<UGameFeaturePluginStateMachine*>& OutDependencyMachines) const
+{
+	FGameFeaturePluginDetails Details;
+	if (GetGameFeaturePluginDetails(PluginFilename, Details))
+	{
+		for (const FGameFeaturePluginReferenceDetails& PluginDependency : Details.PluginDependencies)
+		{
+			if (PluginDependency.bShouldActivate)
+			{
+				const FString& DependencyName = PluginDependency.PluginName;
+				TValueOrError<FString, FString> DependencyURLInfo = GameSpecificPolicies->ResolvePluginDependency(PluginURL, DependencyName);
+				if (DependencyURLInfo.HasError())
+				{
+					UE_LOG(LogGameFeatures, Error, TEXT("Failure to resolve dependency %s [%s] for parent plugin url: %s"), *DependencyName, *DependencyURLInfo.GetError(), *PluginURL);
+					return false;
+				}
+
+				const FString& DependencyURL = DependencyURLInfo.GetValue();
+
+				// Dependency may not be a GFP and so will have an empty URL but not have an error
+				if (DependencyURL.IsEmpty())
+				{
+					continue;
+				}
+
+				UGameFeaturePluginStateMachine* Dependency = FindGameFeaturePluginStateMachine(DependencyURL);
+				if (Dependency)
+				{
+					OutDependencyMachines.Add(Dependency);
+				}
+				else
+				{
+					// Depenedency may have been fully terminated which is considered deactivated already.
+					UE_LOG(LogGameFeatures, Log, TEXT("FindPluginDependencyStateMachinesToDeactivate unable to find plugin state machine for %s using URL %s"), *DependencyName, *DependencyURL);
+				}
+			}
+		}
+
+		return true;
+	}
+
+	return false;
+}
+
 void UGameFeaturesSubsystem::ListGameFeaturePlugins(const TArray<FString>& Args, UWorld* InWorld, FOutputDevice& Ar)
 {
 	const bool bAlphaSort = Args.ContainsByPredicate([](const FString& Arg) { return Arg.Compare(TEXT("-ALPHASORT"), ESearchCase::IgnoreCase) == 0; });
diff --git a/Engine/Plugins/Experimental/GameFeatures/Source/GameFeatures/Public/GameFeatureTypes.h b/Engine/Plugins/Experimental/GameFeatures/Source/GameFeatures/Public/GameFeatureTypes.h
--- a/Engine/Plugins/Experimental/GameFeatures/Source/GameFeatures/Public/GameFeatureTypes.h
+++ b/Engine/Plugins/Experimental/GameFeatures/Source/GameFeatures/Public/GameFeatureTypes.h
@@ -7,38 +7,40 @@
 class FString;
 
 #define GAME_FEATURE_PLUGIN_STATE_LIST(XSTATE)	\
-	XSTATE(Uninitialized,				NSLOCTEXT("GameFeatures", "UninitializedStateDisplayName", "Uninitialized"))								/* Unset. Not yet been set up. */ \
-	XSTATE(Terminal,					NSLOCTEXT("GameFeatures", "TerminalStateDisplayName", "Terminal"))											/* Final State before removal of the state machine. */ \
-	XSTATE(UnknownStatus,				NSLOCTEXT("GameFeatures", "UnknownStatusStateDisplayName", "UnknownStatus"))								/* Initialized, but the only thing known is the URL to query status. */ \
-	XSTATE(Uninstalled,					NSLOCTEXT("GameFeatures", "UninstalledStateDisplayName", "Uninstalled"))									/* All installed data for this plugin has now been uninstalled from local storage (i.e the hard drive) */ \
-	XSTATE(Uninstalling,				NSLOCTEXT("GameFeatures", "UninstallingStateDisplayName", "Uninstalling"))									/* Transition state between StatusKnown -> Terminal for any plugin that can have data that needs to have local data uninstalled. */ \
-	XSTATE(ErrorUninstalling,			NSLOCTEXT("GameFeatures", "ErrorUninstallingStateDisplayName", "ErrorUninstalling"))						/* Error state for Uninstalling -> Terminal transition. */  \
-	XSTATE(CheckingStatus,				NSLOCTEXT("GameFeatures", "CheckingStatusStateDisplayName", "CheckingStatus"))								/* Transition state UnknownStatus -> StatusKnown. The status is in the process of being queried. */ \
-	XSTATE(ErrorCheckingStatus,			NSLOCTEXT("GameFeatures", "ErrorCheckingStatusStateDisplayName", "ErrorCheckingStatus"))					/* Error state for UnknownStatus -> StatusKnown transition. */ \
-	XSTATE(ErrorUnavailable,			NSLOCTEXT("GameFeatures", "ErrorUnavailableStateDisplayName", "ErrorUnavailable"))							/* Error state for UnknownStatus -> StatusKnown transition. */ \
-	XSTATE(StatusKnown,					NSLOCTEXT("GameFeatures", "StatusKnownStateDisplayName", "StatusKnown"))									/* The plugin's information is known, but no action has taken place yet. */ \
-	XSTATE(Releasing,					NSLOCTEXT("GameFeatures", "ReleasingStateDisplayName", "Releasing"))										/* Transition State for Installed -> StatusKnown. Releases local data from any relevant caches. */ \
-	XSTATE(ErrorManagingData,			NSLOCTEXT("GameFeatures", "ErrorManagingDataStateDisplayName", "ErrorManagingData"))						/* Error state for Installed -> StatusKnown and StatusKnown -> Installed transitions. */ \
-	XSTATE(Downloading,					NSLOCTEXT("GameFeatures", "DownloadingStateDisplayName", "Downloading"))									/* Transition state StatusKnown -> Installed. In the process of adding to local storage. */ \
-	XSTATE(Installed,					NSLOCTEXT("GameFeatures", "InstalledStateDisplayName", "Installed"))										/* The plugin is in local storage (i.e. it is on the hard drive) */ \
-	XSTATE(ErrorMounting,				NSLOCTEXT("GameFeatures", "ErrorMountingStateDisplayName", "ErrorMounting"))								/* Error state for Installed -> Registered and Registered -> Installed transitions. */ \
-	XSTATE(ErrorWaitingForDependencies,	NSLOCTEXT("GameFeatures", "ErrorWaitingForDependenciesStateDisplayName", "ErrorWaitingForDependencies"))	/* Error state for Installed -> Registered and Registered -> Installed transitions. */ \
-	XSTATE(ErrorRegistering,			NSLOCTEXT("GameFeatures", "ErrorRegisteringDisplayName", "ErrorRegistering"))								/* Error state for Installed -> Registered and Registered -> Installed transitions. */ \
-	XSTATE(WaitingForDependencies,		NSLOCTEXT("GameFeatures", "WaitingForDependenciesStateDisplayName", "WaitingForDependencies"))				/* Transition state Installed -> Registered. In the process of loading code/content for all dependencies into memory. */ \
-	XSTATE(Unmounting,					NSLOCTEXT("GameFeatures", "UnmountingStateDisplayName", "Unmounting"))										/* Transition state Registered -> Installed. The content file(s) (i.e. pak file) for the plugin is unmounting. */ \
-	XSTATE(Mounting,					NSLOCTEXT("GameFeatures", "MountingStateDisplayName", "Mounting"))											/* Transition state Installed -> Registered. The content files(s) (i.e. pak file) for the plugin is getting mounted. */ \
-	XSTATE(Unregistering,				NSLOCTEXT("GameFeatures", "UnregisteringStateDisplayName", "Unregistering"))								/* Transition state Registered -> Installed. Cleaning up data gathered in Registering. */ \
-	XSTATE(Registering,					NSLOCTEXT("GameFeatures", "RegisteringStateDisplayName", "Registering"))									/* Transition state Installed -> Registered. Discovering assets in the plugin, but not loading them, except a few for discovery reasons. */ \
-	XSTATE(Registered,					NSLOCTEXT("GameFeatures", "RegisteredStateDisplayName", "Registered"))										/* The assets in the plugin are known, but have not yet been loaded, except a few for discovery reasons. */ \
-	XSTATE(ErrorLoading,				NSLOCTEXT("GameFeatures", "ErrorLoadingDisplayName", "ErrorLoading"))										/* Error state for Loading -> Loaded transition */ \
-	XSTATE(Unloading,					NSLOCTEXT("GameFeatures", "UnloadingStateDisplayName", "Unloading"))										/* Transition state Loaded -> Registered. In the process of removing code/content from memory. */ \
-	XSTATE(Loading,						NSLOCTEXT("GameFeatures", "LoadingStateDisplayName", "Loading"))											/* Transition state Registered -> Loaded. In the process of loading code/content into memory. */ \
-	XSTATE(Loaded,						NSLOCTEXT("GameFeatures", "LoadedStateDisplayName", "Loaded"))												/* The plugin is loaded into memory, but not registered with game systems and active. */ \
-	XSTATE(ErrorActivatingDependencies,	NSLOCTEXT("GameFeatures", "ErrorActivatingDependenciesStateDisplayName", "ErrorActivatingDependencies"))	/* Error state for Registered -> Active transition. */ \
-	XSTATE(ActivatingDependencies,		NSLOCTEXT("GameFeatures", "ActivatingDependenciesStateDisplayName", "ActivatingDependencies"))				/* Transition state Registered -> Active. In the process of selectively activating dependencies.*/ \
-	XSTATE(Deactivating,				NSLOCTEXT("GameFeatures", "DeactivatingStateDisplayName", "Deactivating"))									/* Transition state Active -> Loaded. Currently unregistering with game systems. */ \
-	XSTATE(Activating,					NSLOCTEXT("GameFeatures", "ActivatingStateDisplayName", "Activating"))										/* Transition state Loaded -> Active. Currently registering plugin code/content with game systems. */ \
-	XSTATE(Active,						NSLOCTEXT("GameFeatures", "ActiveStateDisplayName", "Active"))												/* Plugin is fully loaded and active. It is affecting the game.  */ 
+	XSTATE(Uninitialized,					NSLOCTEXT("GameFeatures", "UninitializedStateDisplayName", "Uninitialized"))									/* Unset. Not yet been set up. */ \
+	XSTATE(Terminal,						NSLOCTEXT("GameFeatures", "TerminalStateDisplayName", "Terminal"))												/* Final State before removal of the state machine. */ \
+	XSTATE(UnknownStatus,					NSLOCTEXT("GameFeatures", "UnknownStatusStateDisplayName", "UnknownStatus"))									/* Initialized, but the only thing known is the URL to query status. */ \
+	XSTATE(Uninstalled,						NSLOCTEXT("GameFeatures", "UninstalledStateDisplayName", "Uninstalled"))										/* All installed data for this plugin has now been uninstalled from local storage (i.e the hard drive) */ \
+	XSTATE(Uninstalling,					NSLOCTEXT("GameFeatures", "UninstallingStateDisplayName", "Uninstalling"))										/* Transition state between StatusKnown -> Terminal for any plugin that can have data that needs to have local data uninstalled. */ \
+	XSTATE(ErrorUninstalling,				NSLOCTEXT("GameFeatures", "ErrorUninstallingStateDisplayName", "ErrorUninstalling"))							/* Error state for Uninstalling -> Terminal transition. */  \
+	XSTATE(CheckingStatus,					NSLOCTEXT("GameFeatures", "CheckingStatusStateDisplayName", "CheckingStatus"))									/* Transition state UnknownStatus -> StatusKnown. The status is in the process of being queried. */ \
+	XSTATE(ErrorCheckingStatus,				NSLOCTEXT("GameFeatures", "ErrorCheckingStatusStateDisplayName", "ErrorCheckingStatus"))						/* Error state for UnknownStatus -> StatusKnown transition. */ \
+	XSTATE(ErrorUnavailable,				NSLOCTEXT("GameFeatures", "ErrorUnavailableStateDisplayName", "ErrorUnavailable"))								/* Error state for UnknownStatus -> StatusKnown transition. */ \
+	XSTATE(StatusKnown,						NSLOCTEXT("GameFeatures", "StatusKnownStateDisplayName", "StatusKnown"))										/* The plugin's information is known, but no action has taken place yet. */ \
+	XSTATE(Releasing,						NSLOCTEXT("GameFeatures", "ReleasingStateDisplayName", "Releasing"))											/* Transition State for Installed -> StatusKnown. Releases local data from any relevant caches. */ \
+	XSTATE(ErrorManagingData,				NSLOCTEXT("GameFeatures", "ErrorManagingDataStateDisplayName", "ErrorManagingData"))							/* Error state for Installed -> StatusKnown and StatusKnown -> Installed transitions. */ \
+	XSTATE(Downloading,						NSLOCTEXT("GameFeatures", "DownloadingStateDisplayName", "Downloading"))										/* Transition state StatusKnown -> Installed. In the process of adding to local storage. */ \
+	XSTATE(Installed,						NSLOCTEXT("GameFeatures", "InstalledStateDisplayName", "Installed"))											/* The plugin is in local storage (i.e. it is on the hard drive) */ \
+	XSTATE(ErrorMounting,					NSLOCTEXT("GameFeatures", "ErrorMountingStateDisplayName", "ErrorMounting"))									/* Error state for Installed -> Registered and Registered -> Installed transitions. */ \
+	XSTATE(ErrorWaitingForDependencies,		NSLOCTEXT("GameFeatures", "ErrorWaitingForDependenciesStateDisplayName", "ErrorWaitingForDependencies"))		/* Error state for Installed -> Registered and Registered -> Installed transitions. */ \
+	XSTATE(ErrorRegistering,				NSLOCTEXT("GameFeatures", "ErrorRegisteringDisplayName", "ErrorRegistering"))									/* Error state for Installed -> Registered and Registered -> Installed transitions. */ \
+	XSTATE(WaitingForDependencies,			NSLOCTEXT("GameFeatures", "WaitingForDependenciesStateDisplayName", "WaitingForDependencies"))					/* Transition state Installed -> Registered. In the process of loading code/content for all dependencies into memory. */ \
+	XSTATE(Unmounting,						NSLOCTEXT("GameFeatures", "UnmountingStateDisplayName", "Unmounting"))											/* Transition state Registered -> Installed. The content file(s) (i.e. pak file) for the plugin is unmounting. */ \
+	XSTATE(Mounting,						NSLOCTEXT("GameFeatures", "MountingStateDisplayName", "Mounting"))												/* Transition state Installed -> Registered. The content files(s) (i.e. pak file) for the plugin is getting mounted. */ \
+	XSTATE(Unregistering,					NSLOCTEXT("GameFeatures", "UnregisteringStateDisplayName", "Unregistering"))									/* Transition state Registered -> Installed. Cleaning up data gathered in Registering. */ \
+	XSTATE(Registering,						NSLOCTEXT("GameFeatures", "RegisteringStateDisplayName", "Registering"))										/* Transition state Installed -> Registered. Discovering assets in the plugin, but not loading them, except a few for discovery reasons. */ \
+	XSTATE(Registered,						NSLOCTEXT("GameFeatures", "RegisteredStateDisplayName", "Registered"))											/* The assets in the plugin are known, but have not yet been loaded, except a few for discovery reasons. */ \
+	XSTATE(ErrorLoading,					NSLOCTEXT("GameFeatures", "ErrorLoadingDisplayName", "ErrorLoading"))											/* Error state for Loading -> Loaded transition */ \
+	XSTATE(Unloading,						NSLOCTEXT("GameFeatures", "UnloadingStateDisplayName", "Unloading"))											/* Transition state Loaded -> Registered. In the process of removing code/content from memory. */ \
+	XSTATE(Loading,							NSLOCTEXT("GameFeatures", "LoadingStateDisplayName", "Loading"))												/* Transition state Registered -> Loaded. In the process of loading code/content into memory. */ \
+	XSTATE(Loaded,							NSLOCTEXT("GameFeatures", "LoadedStateDisplayName", "Loaded"))													/* The plugin is loaded into memory, but not registered with game systems and active. */ \
+	XSTATE(ErrorActivatingDependencies,		NSLOCTEXT("GameFeatures", "ErrorActivatingDependenciesStateDisplayName", "ErrorActivatingDependencies"))		/* Error state for Registered -> Active transition. */ \
+	XSTATE(ActivatingDependencies,			NSLOCTEXT("GameFeatures", "ActivatingDependenciesStateDisplayName", "ActivatingDependencies"))					/* Transition state Registered -> Active. In the process of selectively activating dependencies.*/ \
+	XSTATE(Deactivating,					NSLOCTEXT("GameFeatures", "DeactivatingStateDisplayName", "Deactivating"))										/* Transition state Active -> Loaded. Currently unregistering with game systems. */ \
+	XSTATE(ErrorDeactivatingDependencies,	NSLOCTEXT("GameFeatures", "ErrorDeactivatingDependenciesStateDisplayName", "DeactivatingDependencies"))			/* Error state for Active -> Loaded transition. */ \
+	XSTATE(DeactivatingDependencies,		NSLOCTEXT("GameFeatures", "DeactivatingDependenciesStateDisplayName", "DeactivatingDependencies"))				/* Transition state Active -> Loaded. In the process of selectively deactivating dependencies.*/ \
+	XSTATE(Activating,						NSLOCTEXT("GameFeatures", "ActivatingStateDisplayName", "Activating"))											/* Transition state Loaded -> Active. Currently registering plugin code/content with game systems. */ \
+	XSTATE(Active,							NSLOCTEXT("GameFeatures", "ActiveStateDisplayName", "Active"))													/* Plugin is fully loaded and active. It is affecting the game.  */ 
 
 #define GAME_FEATURE_PLUGIN_STATE_ENUM(inEnum, inText) inEnum,
 namespace GameFeaturePluginStatePrivate
diff --git a/Engine/Plugins/Experimental/GameFeatures/Source/GameFeatures/Public/GameFeaturesSubsystem.h b/Engine/Plugins/Experimental/GameFeatures/Source/GameFeatures/Public/GameFeaturesSubsystem.h
--- a/Engine/Plugins/Experimental/GameFeatures/Source/GameFeatures/Public/GameFeaturesSubsystem.h
+++ b/Engine/Plugins/Experimental/GameFeatures/Source/GameFeatures/Public/GameFeaturesSubsystem.h
@@ -670,6 +670,9 @@ private:
 	bool FindPluginDependencyStateMachinesToActivate(const FString& PluginURL, const FString& PluginFilename, TArray<UGameFeaturePluginStateMachine*>& OutDependencyMachines) const;
 	friend struct FActivatingDependenciesTransitionPolicy;
 
+	bool FindPluginDependencyStateMachinesToDeactivate(const FString& PluginURL, const FString& PluginFilename, TArray<UGameFeaturePluginStateMachine*>& OutDependencyMachines) const;
+	friend struct FDeactivatingDependenciesTransitionPolicy;
+
 	/** Handle 'ListGameFeaturePlugins' console command */
 	void ListGameFeaturePlugins(const TArray<FString>& Args, UWorld* InWorld, FOutputDevice& Ar);
 
diff --git a/Engine/Plugins/Experimental/GameFeatures/Source/GameFeaturesEditor/Private/SGameFeatureStateWidget.cpp b/Engine/Plugins/Experimental/GameFeatures/Source/GameFeaturesEditor/Private/SGameFeatureStateWidget.cpp
--- a/Engine/Plugins/Experimental/GameFeatures/Source/GameFeaturesEditor/Private/SGameFeatureStateWidget.cpp
+++ b/Engine/Plugins/Experimental/GameFeatures/Source/GameFeaturesEditor/Private/SGameFeatureStateWidget.cpp
@@ -65,7 +65,7 @@ FText SGameFeatureStateWidget::GetDisplayNameOfState(EGameFeaturePluginState Sta
 
 FText SGameFeatureStateWidget::GetTooltipOfState(EGameFeaturePluginState State)
 {
-	static_assert((int32)EGameFeaturePluginState::MAX == 32, "");
+	static_assert((int32)EGameFeaturePluginState::MAX == 34, "");
 
 	switch (State)
 	{
@@ -127,6 +127,10 @@ FText SGameFeatureStateWidget::GetTooltipOfState(EGameFeaturePluginState State)
 		return LOCTEXT("StateTooltip_ErrorActivateDependencies", "Error state for Loaded -> Active and Active -> Loaded transitions.");
 	case EGameFeaturePluginState::ActivatingDependencies:
 		return LOCTEXT("StateTooltip_ActivateDependencies", "Transition state Loaded -> Active. In the process of selectively activating dependencies.");
+	case EGameFeaturePluginState::ErrorDeactivatingDependencies:
+		return LOCTEXT("StateTooltip_ErrorDeactivatingDependencies", "Error state for Active -> Loaded transition.");
+	case EGameFeaturePluginState::DeactivatingDependencies:
+		return LOCTEXT("StateTooltip_DeactivateDependencies", "Transition state Active -> Loaded. In the process of selectively deactivating dependencies.");
 	case EGameFeaturePluginState::Deactivating:
 		return LOCTEXT("StateTooltip_Deactivating", "Transition state Active -> Loaded. Currently unregistering with game systems.");
 	case EGameFeaturePluginState::Activating:
