Create an asset registry bin for each GPF Chunk.
- This is similar to what we do for encrypted paks
- Should allow for reloading GFP AR State

Need to some more work to fully verify this but seems to work in initial testing, and I'd like to get eyes on it.

diff --git a/Engine/Plugins/Experimental/GameFeatures/Source/GameFeatures/Private/GameFeatureAction_AddChunkOverride.cpp b/Engine/Plugins/Experimental/GameFeatures/Source/GameFeatures/Private/GameFeatureAction_AddChunkOverride.cpp
--- a/Engine/Plugins/Experimental/GameFeatures/Source/GameFeatures/Private/GameFeatureAction_AddChunkOverride.cpp
+++ b/Engine/Plugins/Experimental/GameFeatures/Source/GameFeatures/Private/GameFeatureAction_AddChunkOverride.cpp
@@ -103,6 +103,11 @@ void UGameFeatureAction_AddChunkOverride::GetChunkForPackage(const FString& Pack
 	}
 }
 
+FString UGameFeatureAction_AddChunkOverride::GetPluginNameFromChunkID(int32 ChunkID)
+{
+	return GameFeatureAction_AddChunkOverride::ChunkIdToPluginMap.FindRef(ChunkID);
+}
+
 void UGameFeatureAction_AddChunkOverride::PostRename(UObject* OldOuter, const FName OldName)
 {
 	Super::PostRename(OldOuter, OldName);
diff --git a/Engine/Plugins/Experimental/GameFeatures/Source/GameFeatures/Private/GameFeaturePluginStateMachine.cpp b/Engine/Plugins/Experimental/GameFeatures/Source/GameFeatures/Private/GameFeaturePluginStateMachine.cpp
--- a/Engine/Plugins/Experimental/GameFeatures/Source/GameFeatures/Private/GameFeaturePluginStateMachine.cpp
+++ b/Engine/Plugins/Experimental/GameFeatures/Source/GameFeatures/Private/GameFeaturePluginStateMachine.cpp
@@ -2135,26 +2135,40 @@ struct FGameFeaturePluginState_Mounting : public FGameFeaturePluginState
 			return;
 		}
 
-		if (StateProperties.GetPluginProtocol() != EGameFeaturePluginProtocol::InstallBundle)
-		{
-			CompletedSubStates |= ESubState::LoadAssetRegistry;
-			return;
-		}
-
 		TRACE_CPUPROFILER_EVENT_SCOPE(GFP_Mounting_AR);
 
 		// After the new plugin is mounted add the asset registry for that plugin.
-
-		const FString PluginFolder = FPaths::GetPath(StateProperties.PluginInstalledFilename);
-		FString PluginAssetRegistry = PluginFolder / TEXT("AssetRegistry.bin");
-
 		TSharedPtr<IPlugin> NewlyMountedPlugin = IPluginManager::Get().FindPlugin(StateProperties.PluginName);
-		if (!NewlyMountedPlugin || !NewlyMountedPlugin->CanContainContent() || !IFileManager::Get().FileExists(*PluginAssetRegistry))
+		if (!NewlyMountedPlugin || !NewlyMountedPlugin->CanContainContent())
 		{
 			CompletedSubStates |= ESubState::LoadAssetRegistry;
 			return;
 		}
 
+		FString PluginAssetRegistry;
+		{
+			const FString PluginFolder = FPaths::GetPath(StateProperties.PluginInstalledFilename);
+			TArray<FString> PluginAssetRegistrySearchPaths;
+			// For GFPs cooked as DLC
+			PluginAssetRegistrySearchPaths.Add(PluginFolder / TEXT("AssetRegistry.bin"));
+			// For GFPs with a unique chunk
+			PluginAssetRegistrySearchPaths.Add(FPaths::ProjectDir() / FString::Printf(TEXT("AssetRegistry_GFP_%s.bin"), *StateProperties.PluginName));
+			for (FString& Path : PluginAssetRegistrySearchPaths)
+			{
+				if (IFileManager::Get().FileExists(*Path))
+				{
+					PluginAssetRegistry = MoveTemp(Path);
+					break;
+				}
+			}
+
+			if (PluginAssetRegistry.IsEmpty())
+			{
+				CompletedSubStates |= ESubState::LoadAssetRegistry;
+				return;
+			}
+		}
+
 		if (!UseAsyncLoading())
 		{
 			FAssetRegistryState PluginAssetRegistryState;
diff --git a/Engine/Plugins/Experimental/GameFeatures/Source/GameFeatures/Public/GameFeatureAction_AddChunkOverride.h b/Engine/Plugins/Experimental/GameFeatures/Source/GameFeatures/Public/GameFeatureAction_AddChunkOverride.h
--- a/Engine/Plugins/Experimental/GameFeatures/Source/GameFeatures/Public/GameFeatureAction_AddChunkOverride.h
+++ b/Engine/Plugins/Experimental/GameFeatures/Source/GameFeatures/Public/GameFeatureAction_AddChunkOverride.h
@@ -37,6 +37,8 @@ public:
 	 */
 	GAMEFEATURES_API static void GetChunkForPackage(const FString& PackageName, const int32 DefaultGameChunk, TArray<int32>& OutChunkList);
 
+	GAMEFEATURES_API static FString GetPluginNameFromChunkID(int32 ChunkID);
+
 	/** UObject overrides */
 	virtual void PostRename(UObject* OldOuter, const FName OldName) override;
 	virtual void PostEditChangeProperty(FPropertyChangedEvent& PropertyChangedEvent) override;
diff --git a/Engine/Source/Editor/UnrealEd/Private/Commandlets/AssetRegistryGenerator.cpp b/Engine/Source/Editor/UnrealEd/Private/Commandlets/AssetRegistryGenerator.cpp
--- a/Engine/Source/Editor/UnrealEd/Private/Commandlets/AssetRegistryGenerator.cpp
+++ b/Engine/Source/Editor/UnrealEd/Private/Commandlets/AssetRegistryGenerator.cpp
@@ -637,19 +637,19 @@ bool FAssetRegistryGenerator::GenerateStreamingInstallManifest(int64 InOverrideC
 			if (Guid.IsValid())
 			{
 				PakChunkOptions += TEXT(" encryptionkeyguid=") + Guid.ToString();
+			}
 
-				// If this chunk has a seperate unique asset registry, add it to first subchunk's manifest here
-				if (SubChunkIndex == 0)
+			// If this chunk has a seperate unique asset registry, add it to first subchunk's manifest here
+			if (SubChunkIndex == 0)
+			{
+				// For chunks with unique asset registry name, pakchunkIndex should equal chunkid
+				FName RegistryName = UAssetManager::Get().GetUniqueAssetRegistryName(PakchunkIndex);
+				if (RegistryName != NAME_None)
 				{
-					// For chunks with unique asset registry name, pakchunkIndex should equal chunkid
-					FName RegistryName = UAssetManager::Get().GetUniqueAssetRegistryName(PakchunkIndex);
-					if (RegistryName != NAME_None)
-					{
-						FString AssetRegistryFilename = FString::Printf(TEXT("%s%sAssetRegistry%s.bin"),
-							*InSandboxFile.GetSandboxDirectory(), *InSandboxFile.GetGameSandboxDirectoryName(),
-							*RegistryName.ToString());
-						ChunkFilenames.Add(AssetRegistryFilename);
-					}
+					FString AssetRegistryFilename = FString::Printf(TEXT("%s%sAssetRegistry%s.bin"),
+						*InSandboxFile.GetSandboxDirectory(), *InSandboxFile.GetGameSandboxDirectoryName(),
+						*RegistryName.ToString());
+					ChunkFilenames.Add(AssetRegistryFilename);
 				}
 			}
 
