Add CookPlugins static function to GameFeaturesSubsystem so the asset manager can know which plugins to cook.

diff --git a/Engine/Plugins/Experimental/GameFeatures/Source/GameFeatures/Private/GameFeaturesSubsystem.cpp b/Engine/Plugins/Experimental/GameFeatures/Source/GameFeatures/Private/GameFeaturesSubsystem.cpp
--- a/Engine/Plugins/Experimental/GameFeatures/Source/GameFeatures/Private/GameFeaturesSubsystem.cpp
+++ b/Engine/Plugins/Experimental/GameFeatures/Source/GameFeatures/Private/GameFeaturesSubsystem.cpp
@@ -2559,3 +2559,21 @@ EGameFeaturePluginState UGameFeaturesSubsystem::ConvertInitialFeatureStateToTarg
 	}
 	return InitialState;
 }
+
+void UGameFeaturesSubsystem::GetPluginsToCook(TSet<FString>& OutPlugins)
+{
+	// Command line parameter -CookPlugins.
+	static TArray<FString> PluginsList = []()
+	{
+		TArray<FString> ReturnList;
+		FString CookPluginsStr;
+		if (FParse::Value(FCommandLine::Get(), TEXT("CookPlugins="), CookPluginsStr, false))
+		{
+			CookPluginsStr.ParseIntoArray(ReturnList, TEXT(","));
+		}
+
+		return ReturnList;
+	}();
+	
+	OutPlugins.Append(PluginsList);
+}
diff --git a/Engine/Plugins/Experimental/GameFeatures/Source/GameFeatures/Public/GameFeaturesSubsystem.h b/Engine/Plugins/Experimental/GameFeatures/Source/GameFeatures/Public/GameFeaturesSubsystem.h
--- a/Engine/Plugins/Experimental/GameFeatures/Source/GameFeatures/Public/GameFeaturesSubsystem.h
+++ b/Engine/Plugins/Experimental/GameFeatures/Source/GameFeatures/Public/GameFeaturesSubsystem.h
@@ -574,6 +574,8 @@ public:
 
 	static EGameFeaturePluginState ConvertInitialFeatureStateToTargetState(EBuiltInAutoState InitialState);
 
+	/** Used during a DLC cook to determine which plugins should be cooked */
+	static void GetPluginsToCook(TSet<FString>& OutPlugins);
 private:
 	TSet<FString> GetActivePluginNames() const;
 
