Don't apply chunk override unless the plugin passes an optional game override.

diff --git a/Engine/Plugins/Experimental/GameFeatures/Source/GameFeatures/Private/GameFeatureAction_AddChunkOverride.cpp b/Engine/Plugins/Experimental/GameFeatures/Source/GameFeatures/Private/GameFeatureAction_AddChunkOverride.cpp
--- a/Engine/Plugins/Experimental/GameFeatures/Source/GameFeatures/Private/GameFeatureAction_AddChunkOverride.cpp
+++ b/Engine/Plugins/Experimental/GameFeatures/Source/GameFeatures/Private/GameFeatureAction_AddChunkOverride.cpp
@@ -21,9 +21,15 @@ namespace GameFeatureAction_AddChunkOverride
 	static TMap<FString, int32> PluginToChunkId;
 }
 
+UGameFeatureAction_AddChunkOverride::FShouldAddChunkOverride UGameFeatureAction_AddChunkOverride::ShouldAddChunkOverride;
+
 void UGameFeatureAction_AddChunkOverride::OnGameFeatureRegistering()
 {
-	AddChunkIdOverride();
+	const bool bShouldAddChunkOverride = ShouldAddChunkOverride.IsBound() ? ShouldAddChunkOverride.Execute(GetTypedOuter<UGameFeatureData>()) : true;
+	if (bShouldAddChunkOverride)
+	{
+		AddChunkIdOverride();
+	}
 }
 
 void UGameFeatureAction_AddChunkOverride::OnGameFeatureUnregistering()
diff --git a/Engine/Plugins/Experimental/GameFeatures/Source/GameFeatures/Public/GameFeatureAction_AddChunkOverride.h b/Engine/Plugins/Experimental/GameFeatures/Source/GameFeatures/Public/GameFeatureAction_AddChunkOverride.h
--- a/Engine/Plugins/Experimental/GameFeatures/Source/GameFeatures/Public/GameFeatureAction_AddChunkOverride.h
+++ b/Engine/Plugins/Experimental/GameFeatures/Source/GameFeatures/Public/GameFeatureAction_AddChunkOverride.h
@@ -26,6 +26,14 @@ public:
 	virtual void OnGameFeatureUnregistering() override;
 	//~End of UGameFeatureAction interface
 
+	DECLARE_DELEGATE_RetVal_OneParam(bool, FShouldAddChunkOverride, const UGameFeatureData*);
+	/**
+	 * Optionally bound delegate to determine when to add the chunk override.
+	 * When bound this will be checked before attempting to add the chunk override.
+	 * Bound delegates should return true if the GameFeatureData should have a chunk id overriden; otherwise, false.
+	 */
+	GAMEFEATURES_API static FShouldAddChunkOverride ShouldAddChunkOverride;
+
 #if WITH_EDITOR
 	/**
 	 * Given the package name will check if this is a package from a GFP that we want to assign a specific chunk to.
