New cvar for using content beacon for playlist downloadable content.
Content beacon now ignores plugins which are managed by LoadBuiltInPlugins.
Content beacon supports installed state (matches to GameFeaturePluginTargetState Registered).

diff --git a/Engine/Plugins/Experimental/GameFeatures/Source/GameFeatures/Private/GameFeaturePluginStateMachine.cpp b/Engine/Plugins/Experimental/GameFeatures/Source/GameFeatures/Private/GameFeaturePluginStateMachine.cpp
--- a/Engine/Plugins/Experimental/GameFeatures/Source/GameFeatures/Private/GameFeaturePluginStateMachine.cpp
+++ b/Engine/Plugins/Experimental/GameFeatures/Source/GameFeatures/Private/GameFeaturePluginStateMachine.cpp
@@ -3357,6 +3357,16 @@ bool UGameFeaturePluginStateMachine::AllowAsyncLoading() const
 	return StateProperties.AllowAsyncLoading();
 }
 
+void UGameFeaturePluginStateMachine::SetWasLoadedAsBuiltIn()
+{
+	StateProperties.bWasLoadedAsBuiltInGameFeaturePlugin = true;
+}
+
+bool UGameFeaturePluginStateMachine::WasLoadedAsBuiltIn() const
+{
+	return StateProperties.bWasLoadedAsBuiltInGameFeaturePlugin;
+}
+
 UGameFeatureData* UGameFeaturePluginStateMachine::GetGameFeatureDataForActivePlugin()
 {
 	if (GetCurrentState() == EGameFeaturePluginState::Active)
diff --git a/Engine/Plugins/Experimental/GameFeatures/Source/GameFeatures/Private/GameFeaturePluginStateMachine.h b/Engine/Plugins/Experimental/GameFeatures/Source/GameFeatures/Private/GameFeaturePluginStateMachine.h
--- a/Engine/Plugins/Experimental/GameFeatures/Source/GameFeatures/Private/GameFeaturePluginStateMachine.h
+++ b/Engine/Plugins/Experimental/GameFeatures/Source/GameFeatures/Private/GameFeaturePluginStateMachine.h
@@ -289,6 +289,9 @@ struct FGameFeaturePluginStateMachineProperties
 	/** Tracks whether or not this state machine added the plugin to the plugin manager. */
 	bool bAddedPluginToManager = false;
 
+	/** Was this plugin loaded using LoadBuiltInGameFeaturePlugin */
+	bool bWasLoadedAsBuiltInGameFeaturePlugin = false;
+
 	/** Whether this state machine should attempt to cancel the current transition */
 	bool bTryCancel = false;
 
@@ -529,6 +532,10 @@ public:
 	/** Whether this machine is allowed to be asynchronous */
 	bool AllowAsyncLoading() const;
 
+	void SetWasLoadedAsBuiltIn();
+
+	bool WasLoadedAsBuiltIn() const;
+
 	/** If the plugin is activated already, we will retrieve its game feature data */
 	UGameFeatureData* GetGameFeatureDataForActivePlugin();
 
diff --git a/Engine/Plugins/Experimental/GameFeatures/Source/GameFeatures/Private/GameFeaturesSubsystem.cpp b/Engine/Plugins/Experimental/GameFeatures/Source/GameFeatures/Private/GameFeaturesSubsystem.cpp
--- a/Engine/Plugins/Experimental/GameFeatures/Source/GameFeatures/Private/GameFeaturesSubsystem.cpp
+++ b/Engine/Plugins/Experimental/GameFeatures/Source/GameFeatures/Private/GameFeaturesSubsystem.cpp
@@ -917,6 +917,15 @@ bool UGameFeaturesSubsystem::IsGameFeaturePluginLoaded(const FString& PluginURL)
 	return false;
 }
 
+bool UGameFeaturesSubsystem::WasGameFeaturePluginLoadedAsBuiltIn(const FString& PluginURL) const
+{
+	if (const UGameFeaturePluginStateMachine* StateMachine = FindGameFeaturePluginStateMachine(PluginURL))
+	{
+		return StateMachine->WasLoadedAsBuiltIn();
+	}
+	return false;
+}
+
 void UGameFeaturesSubsystem::LoadGameFeaturePlugin(const FString& PluginURL, const FGameFeaturePluginLoadComplete& CompleteDelegate)
 {
 	LoadGameFeaturePlugin(PluginURL, FGameFeatureProtocolOptions(), CompleteDelegate);
@@ -1474,6 +1483,8 @@ void UGameFeaturesSubsystem::LoadBuiltInGameFeaturePlugin(const TSharedRef<IPlug
 				}
 				const EGameFeaturePluginState DestinationState = ConvertInitialFeatureStateToTargetState(InitialAutoState);
 
+				StateMachine->SetWasLoadedAsBuiltIn();
+
 				// If we're already at the destination or beyond, don't transition back
 				FGameFeaturePluginStateRange Destination(DestinationState, EGameFeaturePluginState::Active);
 				ChangeGameFeatureDestination(StateMachine, ProtocolOptions, Destination,
diff --git a/Engine/Plugins/Experimental/GameFeatures/Source/GameFeatures/Public/GameFeaturesSubsystem.h b/Engine/Plugins/Experimental/GameFeatures/Source/GameFeatures/Public/GameFeaturesSubsystem.h
--- a/Engine/Plugins/Experimental/GameFeatures/Source/GameFeatures/Public/GameFeaturesSubsystem.h
+++ b/Engine/Plugins/Experimental/GameFeatures/Source/GameFeatures/Public/GameFeaturesSubsystem.h
@@ -459,6 +459,9 @@ public:
 	/** Determines if a plugin is in the Loaded state (or beyond) */
 	bool IsGameFeaturePluginLoaded(const FString& PluginURL) const;
 
+	/** Was this game feature plugin loaded using the LoadBuiltInGameFeaturePlugin path */
+	bool WasGameFeaturePluginLoadedAsBuiltIn(const FString& PluginURL) const;
+
 	/** Loads a single game feature plugin. */
 	void LoadGameFeaturePlugin(const FString& PluginURL, const FGameFeaturePluginLoadComplete& CompleteDelegate);
 	void LoadGameFeaturePlugin(const FString& PluginURL, const FGameFeatureProtocolOptions& ProtocolOptions, const FGameFeaturePluginLoadComplete& CompleteDelegate);
