Change the default game feature policy to synchronously load built in plugins so they are ready by the end of startup
This behavior changed in CL 28626829 when GameFeaturePlugin.AsyncLoad was enabled by default (and plugins loaded after startup are still async by default)

diff --git a/Engine/Plugins/Experimental/GameFeatures/Source/GameFeatures/Private/GameFeaturesProjectPolicies.cpp b/Engine/Plugins/Experimental/GameFeatures/Source/GameFeatures/Private/GameFeaturesProjectPolicies.cpp
--- a/Engine/Plugins/Experimental/GameFeatures/Source/GameFeatures/Private/GameFeaturesProjectPolicies.cpp
+++ b/Engine/Plugins/Experimental/GameFeatures/Source/GameFeatures/Private/GameFeaturesProjectPolicies.cpp
@@ -13,6 +13,9 @@ void UDefaultGameFeaturesProjectPolicies::InitGameFeatureManager()
 
 	auto AdditionalFilter = [&](const FString& PluginFilename, const FGameFeaturePluginDetails& PluginDetails, FBuiltInGameFeaturePluginBehaviorOptions& OutOptions) -> bool
 	{
+		// By default, force all initially loaded plugins to synchronously load, this overrides the behavior of GameFeaturePlugin.AsyncLoad which will be used for later loads
+		OutOptions.bForceSyncLoading = true;
+
 		return true;
 	};
 
