Added 'LexToString' function for EBuiltInAutoState in game feature subsystem.

diff --git a/Engine/Plugins/Experimental/GameFeatures/Source/GameFeatures/Private/GameFeaturesSubsystem.cpp b/Engine/Plugins/Experimental/GameFeatures/Source/GameFeatures/Private/GameFeaturesSubsystem.cpp
--- a/Engine/Plugins/Experimental/GameFeatures/Source/GameFeatures/Private/GameFeaturesSubsystem.cpp
+++ b/Engine/Plugins/Experimental/GameFeatures/Source/GameFeatures/Private/GameFeaturesSubsystem.cpp
@@ -106,6 +106,26 @@ namespace GameFeaturePluginStatePrivate
 }
 #undef GAME_FEATURE_PLUGIN_STATE_LEX_TO_STRING
 
+const FString LexToString(const EBuiltInAutoState BuiltInAutoState)
+{
+	switch (BuiltInAutoState)
+	{
+		case EBuiltInAutoState::Invalid:
+			return TEXT("Invalid");
+		case EBuiltInAutoState::Installed:
+			return TEXT("Installed");
+		case EBuiltInAutoState::Registered:
+			return TEXT("Registered");
+		case EBuiltInAutoState::Loaded:
+			return TEXT("Loaded");
+		case EBuiltInAutoState::Active:
+			return TEXT("Active");
+		default:
+			check(false);
+			return TEXT("Unknown");
+	}
+}
+
 const FString LexToString(const EGameFeatureTargetState GameFeatureTargetState)
 {
 	switch (GameFeatureTargetState)
diff --git a/Engine/Plugins/Experimental/GameFeatures/Source/GameFeatures/Public/GameFeaturesSubsystem.h b/Engine/Plugins/Experimental/GameFeatures/Source/GameFeatures/Public/GameFeaturesSubsystem.h
--- a/Engine/Plugins/Experimental/GameFeatures/Source/GameFeatures/Public/GameFeaturesSubsystem.h
+++ b/Engine/Plugins/Experimental/GameFeatures/Source/GameFeatures/Public/GameFeaturesSubsystem.h
@@ -215,6 +215,7 @@ enum class EBuiltInAutoState : uint8
 	Loaded,
 	Active
 };
+const FString GAMEFEATURES_API LexToString(const EBuiltInAutoState BuiltInAutoState);
 
 UENUM(BlueprintType)
 enum class EGameFeatureTargetState : uint8
