Reorganize device profiles, fix runtime device profiles that could be duplicated, allow explicit platform DP lookups

#rnx

diff --git a/Engine/Plugins/Experimental/GameFeatures/Source/GameFeatures/Private/GameFeatureData.cpp b/Engine/Plugins/Experimental/GameFeatures/Source/GameFeatures/Private/GameFeatureData.cpp
--- a/Engine/Plugins/Experimental/GameFeatures/Source/GameFeatures/Private/GameFeatureData.cpp
+++ b/Engine/Plugins/Experimental/GameFeatures/Source/GameFeatures/Private/GameFeatureData.cpp
@@ -182,6 +182,8 @@ void UGameFeatureData::InitializeHierarchicalPluginIniFiles(const FString& Plugi
 	{
 		TMap<FString, FConfigSection> ConfigSectionsToAdd;
 
+		TArray<FString> ResultingProfiles;
+
 		for (auto& Section : AsConst(PluginConfig))
 		{
 			FString RuleName, ParentClass;
@@ -251,7 +253,7 @@ void UGameFeatureData::InitializeHierarchicalPluginIniFiles(const FString& Plugi
 						TArray<FString> LoadableProfileNames = DeviceProfileManager.GetLoadableProfileNames(*PlatformName);
 						for (const FString& ProfileName : LoadableProfileNames)
 						{
-							DeviceProfileManager.FindProfile(ProfileName, true);
+							DeviceProfileManager.FindProfile(ProfileName, true, *PlatformName);
 						}
 
 						for (const UDeviceProfile* Profile : DeviceProfileManager.Profiles)
@@ -268,10 +270,20 @@ void UGameFeatureData::InitializeHierarchicalPluginIniFiles(const FString& Plugi
 							// Create the config for a runtime profile
 							if (bProfileHasCompatibleParent && !Profile->GetName().EndsWith(ProfileSuffix->GetValue()))
 							{
+								// Ignore duplicates: only the first match in a given config file will be accepted
+								const FString FinalProfileName = Profile->GetName() + ProfileSuffix->GetValue();
+								if (ResultingProfiles.Contains(FinalProfileName))
+								{
+									UE_LOG(LogGameFeatures, Display, TEXT("Ignoring profile %s that has already been overriden as %s"), *Profile->GetName(), *FinalProfileName);
+									continue;
+								}
+
 								FConfigSection RuntimeProfile;
 								RuntimeProfile.Add("DeviceType", PlatformName);
 								RuntimeProfile.Add("BaseProfileName", FConfigValue(Profile->GetName()));
 
+								UE_LOG(LogGameFeatures, Display, TEXT("Creating override for base profile %s"), *Profile->GetName());
+
 								// Inject the parent's matched fragments into the config, if any
 								if (Profile->GetName().Contains("MatchedFragments"))
 								{
@@ -322,7 +334,8 @@ void UGameFeatureData::InitializeHierarchicalPluginIniFiles(const FString& Plugi
 									}
 								}
 
-								ConfigSectionsToAdd.Add(Profile->GetName() + ProfileSuffix->GetValue() + TEXT(" ") + UDeviceProfile::StaticClass()->GetName(), RuntimeProfile);
+								ConfigSectionsToAdd.Add(FinalProfileName + TEXT(" ") + UDeviceProfile::StaticClass()->GetName(), RuntimeProfile);
+								ResultingProfiles.Add(FinalProfileName);
 							}
 						}
 					}
@@ -384,7 +397,7 @@ void UGameFeatureData::InitializeHierarchicalPluginIniFiles(const FString& Plugi
 				if (DeviceType)
 				{
 					UE_LOG(LogGameFeatures, Display, TEXT("Game feature '%s' adding new device profile %s"), *PluginName, *ProfileName);
-					DeviceProfileManager.CreateProfile(ProfileName, DeviceType->GetValue());
+					DeviceProfileManager.CreateProfile(ProfileName, DeviceType->GetValue(), FString(), *PlatformName);
 				}
 			}
 		}
diff --git a/Engine/Source/Runtime/Engine/Classes/DeviceProfiles/DeviceProfileManager.h b/Engine/Source/Runtime/Engine/Classes/DeviceProfiles/DeviceProfileManager.h
--- a/Engine/Source/Runtime/Engine/Classes/DeviceProfiles/DeviceProfileManager.h
+++ b/Engine/Source/Runtime/Engine/Classes/DeviceProfiles/DeviceProfileManager.h
@@ -104,9 +104,11 @@ public:
 	 * Find a profile based on the name.
 	 *
 	 * @param ProfileName - The profile name to find.
+	 * @param bCreateProfileOnFail - Whether to create the profile from config if the object doesn't exist yet.
+	 * @param OptionalPlatformName - The platform name to use for loading.
 	 * @return The found profile.
 	 */
-	ENGINE_API UDeviceProfile* FindProfile( const FString& ProfileName, bool bCreateProfileOnFail = true );
+	ENGINE_API UDeviceProfile* FindProfile(const FString& ProfileName, bool bCreateProfileOnFail = true, FName OptionalPlatformName = FName());
 
 	/**
 	* Overrides the device profile. The original profile can be restored with RestoreDefaultDeviceProfile
diff --git a/Engine/Source/Runtime/Engine/Private/DeviceProfiles/DeviceProfileManager.cpp b/Engine/Source/Runtime/Engine/Private/DeviceProfiles/DeviceProfileManager.cpp
--- a/Engine/Source/Runtime/Engine/Private/DeviceProfiles/DeviceProfileManager.cpp
+++ b/Engine/Source/Runtime/Engine/Private/DeviceProfiles/DeviceProfileManager.cpp
@@ -846,7 +846,7 @@ void UDeviceProfileManager::DeleteProfile( UDeviceProfile* Profile )
 }
 
 
-UDeviceProfile* UDeviceProfileManager::FindProfile( const FString& ProfileName, bool bCreateProfileOnFail )
+UDeviceProfile* UDeviceProfileManager::FindProfile(const FString& ProfileName, bool bCreateProfileOnFail, FName OptionalPlatformName)
 {
 	UDeviceProfile* FoundProfile = nullptr;
 
@@ -862,7 +862,8 @@ UDeviceProfile* UDeviceProfileManager::FindProfile( const FString& ProfileName,
 
 	if ( bCreateProfileOnFail && FoundProfile == nullptr )
 	{
-		FoundProfile = CreateProfile(ProfileName, FPlatformProperties::IniPlatformName());
+		FString PlatformName = OptionalPlatformName != NAME_None ? OptionalPlatformName.ToString() : FString(FPlatformProperties::IniPlatformName());
+		FoundProfile = CreateProfile(ProfileName, PlatformName);
 	}
 	return FoundProfile;
 }
