GFP Predownload Support
- This doesn't instatiate a GFP state machine and doesn't leave install bundles in a reserved cache state when it finishes.  This makes ideal to use in situations where a download cache pre-warm is desired but actual GFP creation should be deffered.
- Fix install bundle manager issue that could cause bundles that skipped mounting to not be properly released to the install bundle cache.

diff --git a/Engine/Plugins/Experimental/DefaultInstallBundleManager/Source/Private/DefaultInstallBundleManager.cpp b/Engine/Plugins/Experimental/DefaultInstallBundleManager/Source/Private/DefaultInstallBundleManager.cpp
--- a/Engine/Plugins/Experimental/DefaultInstallBundleManager/Source/Private/DefaultInstallBundleManager.cpp
+++ b/Engine/Plugins/Experimental/DefaultInstallBundleManager/Source/Private/DefaultInstallBundleManager.cpp
@@ -1292,6 +1292,12 @@ void FDefaultInstallBundleManager::UpdateBundleSources(FContentReleaseRequestRef
 		return;
 	}
 
+	// Release from any caches that were reserved
+	for (const TPair<FName, TSharedRef<FInstallBundleCache>>& Pair : BundleCaches)
+	{
+		Pair.Value->Release(Request->BundleName);
+	}
+
 	for (const TPair<EInstallBundleSourceType, TSharedPtr<IInstallBundleSource>>& Pair : BundleSources)
 	{
 		if (const FName* CacheName = BundleSourceCaches.Find(Pair.Key))
@@ -1582,12 +1588,6 @@ void FDefaultInstallBundleManager::UnmountPaks(FContentReleaseRequestRef Request
 
 	SetBundleStatus(BundleInfo, EBundleState::NeedsMount);
 
-	// Release from any caches that were reserved
-	for (const TPair<FName, TSharedRef<FInstallBundleCache>>& Pair : BundleCaches)
-	{
-		Pair.Value->Release(Request->BundleName);
-	}
-
 	Request->StepResult = EContentRequestStepResult::Done;
 }
 
@@ -3078,9 +3078,27 @@ TValueOrError<FInstallBundleRequestInfo, EInstallBundleResult> FDefaultInstallBu
 			// canceled release has finished.
 			if (!bCanceledRelease && EnumHasAnyFlags(Flags, EInstallBundleRequestFlags::SkipMount) && GetBundleStatus(*BundleInfo) == EBundleState::NeedsMount)
 			{
-				RetInfo.InfoFlags |= EInstallBundleRequestInfoFlags::SkippedAlreadyUpdatedBundles;
-				LOG_INSTALL_BUNDLE_MAN_OVERRIDE(LogVerbosityOverride, Verbose, TEXT("RequestUpdateContent Bundle %s  - Already Updated"), *BundleInfo->BundleNameString);
-				bIsFinished = true;
+				// If this bundle is not reserved in a cache, an  install request cannot be skipped
+				bool bNeedsCacheReserve = false;
+				for (const FBundleSourceRelevance& SourceRelevance : BundleInfo->ContributingSources)
+				{
+					if (FName* CacheName = BundleSourceCaches.Find(SourceRelevance.SourceType))
+					{
+						const TSharedRef<FInstallBundleCache>& BundleCache = BundleCaches.FindChecked(*CacheName);
+						if (BundleCache->Contains(BundleName) && !BundleCache->IsReserved(BundleName))
+						{
+							bNeedsCacheReserve = true;
+							break;
+						}
+					}
+				}
+
+				if (!bNeedsCacheReserve)
+				{
+					RetInfo.InfoFlags |= EInstallBundleRequestInfoFlags::SkippedAlreadyUpdatedBundles;
+					LOG_INSTALL_BUNDLE_MAN_OVERRIDE(LogVerbosityOverride, Verbose, TEXT("RequestUpdateContent Bundle %s  - Already Updated"), *BundleInfo->BundleNameString);
+					bIsFinished = true;
+				}
 			}
 			// No need to check bCanceledRelease here.  Unmounting is not Async so if we canceled it early enough we will remain mounted
 			else if (!GetMustWaitForPSOCache(*BundleInfo) && GetBundleStatus(*BundleInfo) == EBundleState::Mounted)
@@ -3433,14 +3451,29 @@ TValueOrError<FInstallBundleReleaseRequestInfo, EInstallBundleResult> FDefaultIn
 		// canceled update has finished.
 		if (ActiveQueuedRequest == nullptr && !bCanceledUpdate)
 		{
-			if (!EnumHasAnyFlags(Flags, EInstallBundleReleaseRequestFlags::RemoveFilesIfPossible) && GetBundleStatus(*BundleInfo) != EBundleState::Mounted)
+			// If this bundle is reserved in a cache, a release request cannot be skipped
+			bool bIsReserved = false;
+			for (const FBundleSourceRelevance& SourceRelevance : BundleInfo->ContributingSources)
+			{
+				if (FName* CacheName = BundleSourceCaches.Find(SourceRelevance.SourceType))
+				{
+					const TSharedRef<FInstallBundleCache>& BundleCache = BundleCaches.FindChecked(*CacheName);
+					if (BundleCache->IsReserved(BundleName))
+					{
+						bIsReserved = true;
+						break;
+					}
+				}
+			}
+
+			if (!bIsReserved && !EnumHasAnyFlags(Flags, EInstallBundleReleaseRequestFlags::RemoveFilesIfPossible) && GetBundleStatus(*BundleInfo) != EBundleState::Mounted)
 			{
 				RetInfo.InfoFlags |= EInstallBundleRequestInfoFlags::SkippedAlreadyReleasedBundles;
 				LOG_INSTALL_BUNDLE_MAN_OVERRIDE(LogVerbosityOverride, Verbose, TEXT("BundlesToRelease Bundle %s  - Already Released"), *BundleInfo->BundleNameString);
 				continue;
 			}
 
-			if (GetBundleStatus(*BundleInfo) == EBundleState::NotInstalled)
+			if (!bIsReserved && GetBundleStatus(*BundleInfo) == EBundleState::NotInstalled)
 			{
 				RetInfo.InfoFlags |= EInstallBundleRequestInfoFlags::SkippedAlreadyRemovedBundles;
 				LOG_INSTALL_BUNDLE_MAN_OVERRIDE(LogVerbosityOverride, Verbose, TEXT("BundlesToRelease Bundle %s  - Already Removed"), *BundleInfo->BundleNameString);
diff --git a/Engine/Plugins/Experimental/GameFeatures/Source/GameFeatures/Private/GameFeaturePluginOperationResult.cpp b/Engine/Plugins/Experimental/GameFeatures/Source/GameFeatures/Private/GameFeaturePluginOperationResult.cpp
--- a/Engine/Plugins/Experimental/GameFeatures/Source/GameFeatures/Private/GameFeaturePluginOperationResult.cpp
+++ b/Engine/Plugins/Experimental/GameFeatures/Source/GameFeatures/Private/GameFeaturePluginOperationResult.cpp
@@ -1,7 +1,10 @@
 // Copyright Epic Games, Inc. All Rights Reserved.
 
 #include "GameFeaturePluginOperationResult.h"
-
+#include "GameFeaturesSubsystem.h" // needed for log category
+#include "InstallBundleTypes.h"
+#include "Internationalization/Internationalization.h"
+#include "Logging/LogMacros.h"
 
 namespace UE::GameFeatures
 {
@@ -34,4 +37,112 @@ namespace UE::GameFeatures
 		}
 		return Out.ToString();
 	}
+
+	namespace CommonErrorCodes
+	{
+		const FText Generic_FatalError = NSLOCTEXT("GameFeatures", "ErrorCodes.GenericFatalError", "A fatal error has occurred installing the game feature. An update to the application may be needed. Please check for updates and restart the application.");
+		const FText Generic_ConnectionError = NSLOCTEXT("GameFeatures", "ErrorCodes.ConnectionGenericError", "An internet connection error has occurred. Please try again later.");
+		const FText Generic_MountError = NSLOCTEXT("GameFeatures", "ErrorCodes.MountGenericError", "An error has occurred loading data for this game feature. Please try again later.");
+
+		const FText BundleResult_NeedsUpdate = NSLOCTEXT("GameFeatures", "ErrorCodes.BundleResult.NeedsUpdate", "An application update is required to install this game feature. Please restart the application after downloading any required updates.");
+		const FText BundleResult_NeedsCacheSpace = NSLOCTEXT("GameFeatures", "ErrorCodes.BundleResult.NeedsCacheSpace", "Unable to allocate enough space in the cache to install this game feature.");
+		const FText BundleResult_NeedsDiskSpace = NSLOCTEXT("GameFeatures", "ErrorCodes.BundleResult.NeedsDiskSpace", "You do not have enough disk space to install this game feature. Please try again after clearing up disk space.");
+		const FText BundleResult_DownloadCancelled = NSLOCTEXT("GameFeatures", "ErrorCodes.BundleResult.DownloadCancelled", "This game feature download was canceled.");
+
+		const FText ReleaseResult_Generic = NSLOCTEXT("GameFeatures", "ErrorCodes.ReleaseResult.Generic", "There was an error uninstalling the content for this game feature. Please restart the application and try again.");
+		const FText ReleaseResult_Cancelled = NSLOCTEXT("GameFeatures", "ErrorCodes.ReleaseResult.Cancelled", "This game feature uninstall was canceled.");
+
+		const FText& GetErrorTextForBundleResult(EInstallBundleResult ErrorResult)
+		{
+			switch (ErrorResult)
+			{
+				//These errors mean an app update is available that we either don't have or failed to get.
+				case EInstallBundleResult::FailedPrereqRequiresLatestClient:
+				case EInstallBundleResult::FailedPrereqRequiresLatestContent:
+				{
+					return BundleResult_NeedsUpdate;
+				}
+
+
+				//These are generally unrecoverable and mean something is seriously wrong with the data for this build
+				case EInstallBundleResult::InitializationError:
+				{
+					return Generic_FatalError;
+				}
+
+				//Not enough space in cache to install the files
+				case EInstallBundleResult::FailedCacheReserve:
+				{
+					return BundleResult_NeedsCacheSpace;
+				}
+
+				//All of these are indicative of not having enough disk space to install the required files
+				case EInstallBundleResult::InstallerOutOfDiskSpaceError:
+				case EInstallBundleResult::ManifestArchiveError:
+				{
+					return BundleResult_NeedsDiskSpace;
+				}
+
+				case EInstallBundleResult::UserCancelledError:
+				{
+					return BundleResult_DownloadCancelled;
+				}
+
+				//Intentionally just show generic error for all these cases
+				case EInstallBundleResult::InstallError:
+				case EInstallBundleResult::ConnectivityError:
+				case EInstallBundleResult::InitializationPending:
+				{
+					return Generic_ConnectionError;
+				}
+
+				//Show generic error for anything missing but log an error
+				default:
+				{
+					UE_LOG(LogGameFeatures, Error, TEXT("Missing error text for EInstallBundleResult %s"), LexToString(ErrorResult));
+					return Generic_ConnectionError;
+				}
+			}
+		}
+
+		const FText& GetErrorTextForReleaseResult(EInstallBundleReleaseResult ErrorResult)
+		{
+			switch (ErrorResult)
+			{
+				case (EInstallBundleReleaseResult::UserCancelledError):
+				{
+					return ReleaseResult_Cancelled;
+				}
+			
+				case (EInstallBundleReleaseResult::ManifestArchiveError):
+				{
+					return ReleaseResult_Generic;
+				}
+
+				default:
+				{
+					//Show generic error for anything missing but log an error
+					UE_LOG(LogGameFeatures, Error, TEXT("Missing error text for EInstallBundleReleaseResult %s"), LexToString(ErrorResult));
+					return ReleaseResult_Generic;
+				}
+			}
+		}
+		const FText& GetGenericFatalError()
+		{
+			return Generic_FatalError;
+		}
+		const FText& GetGenericConnectionError()
+		{
+			return Generic_ConnectionError;
+		}
+		const FText& GetGenericMountError()
+		{
+			return Generic_MountError;
+		}
+		const FText& GetGenericReleaseResult()
+		{
+			return ReleaseResult_Generic;
+		}
+	}
+
 }	// namespace UE::GameFeatures
\ No newline at end of file
diff --git a/Engine/Plugins/Experimental/GameFeatures/Source/GameFeatures/Private/GameFeaturePluginStateMachine.cpp b/Engine/Plugins/Experimental/GameFeatures/Source/GameFeatures/Private/GameFeaturePluginStateMachine.cpp
--- a/Engine/Plugins/Experimental/GameFeatures/Source/GameFeatures/Private/GameFeaturePluginStateMachine.cpp
+++ b/Engine/Plugins/Experimental/GameFeatures/Source/GameFeatures/Private/GameFeaturePluginStateMachine.cpp
@@ -141,97 +141,6 @@ namespace UE::GameFeatures
 	}
 #undef GAME_FEATURE_PLUGIN_PROTOCOL_PREFIX
 
-	namespace CommonErrorCodes
-	{
-		const FText Generic_FatalError = NSLOCTEXT("GameFeatures", "ErrorCodes.GenericFatalError", "A fatal error has occurred installing the game feature. An update to the application may be needed. Please check for updates and restart the application.");
-		const FText Generic_ConnectionError = NSLOCTEXT("GameFeatures", "ErrorCodes.ConnectionGenericError", "An internet connection error has occurred. Please try again later.");
-		const FText Generic_MountError = NSLOCTEXT("GameFeatures", "ErrorCodes.MountGenericError", "An error has occurred loading data for this game feature. Please try again later.");
-
-		const FText BundleResult_NeedsUpdate = NSLOCTEXT("GameFeatures", "ErrorCodes.BundleResult.NeedsUpdate", "An application update is required to install this game feature. Please restart the application after downloading any required updates.");
-		const FText BundleResult_NeedsCacheSpace = NSLOCTEXT("GameFeatures", "ErrorCodes.BundleResult.NeedsCacheSpace", "Unable to allocate enough space in the cache to install this game feature.");
-		const FText BundleResult_NeedsDiskSpace = NSLOCTEXT("GameFeatures", "ErrorCodes.BundleResult.NeedsDiskSpace", "You do not have enough disk space to install this game feature. Please try again after clearing up disk space.");
-		const FText BundleResult_DownloadCancelled = NSLOCTEXT("GameFeatures", "ErrorCodes.BundleResult.DownloadCancelled", "This game feature download was canceled.");
-
-		const FText ReleaseResult_Generic = NSLOCTEXT("GameFeatures", "ErrorCodes.ReleaseResult.Generic", "There was an error uninstalling the content for this game feature. Please restart the application and try again.");
-		const FText ReleaseResult_Cancelled = NSLOCTEXT("GameFeatures", "ErrorCodes.ReleaseResult.Cancelled", "This game feature uninstall was canceled.");
-
-		static const FText& GetErrorTextForBundleResult(EInstallBundleResult ErrorResult)
-		{
-			switch (ErrorResult)
-			{
-				//These errors mean an app update is available that we either don't have or failed to get.
-				case EInstallBundleResult::FailedPrereqRequiresLatestClient:
-				case EInstallBundleResult::FailedPrereqRequiresLatestContent:
-				{
-					return BundleResult_NeedsUpdate;
-				}
-
-
-				//These are generally unrecoverable and mean something is seriously wrong with the data for this build
-				case EInstallBundleResult::InitializationError:
-				{
-					return Generic_FatalError;
-				}
-
-				//Not enough space in cache to install the files
-				case EInstallBundleResult::FailedCacheReserve:
-				{
-					return BundleResult_NeedsCacheSpace;
-				}
-
-				//All of these are indicative of not having enough disk space to install the required files
-				case EInstallBundleResult::InstallerOutOfDiskSpaceError:
-				case EInstallBundleResult::ManifestArchiveError:
-				{
-					return BundleResult_NeedsDiskSpace;
-				}
-
-				case EInstallBundleResult::UserCancelledError:
-				{
-					return BundleResult_DownloadCancelled;
-				}
-
-				//Intentionally just show generic error for all these cases
-				case EInstallBundleResult::InstallError:
-				case EInstallBundleResult::ConnectivityError:
-				case EInstallBundleResult::InitializationPending:
-				{
-					return Generic_ConnectionError;
-				}
-
-				//Show generic error for anything missing but log an error
-				default:
-				{
-					UE_LOG(LogGameFeatures, Error, TEXT("Missing error text for EInstallBundleResult %s"), LexToString(ErrorResult));
-					return Generic_ConnectionError;
-				}
-			}
-		}
-
-		static const FText& GetErrorTextForReleaseResult(EInstallBundleReleaseResult ErrorResult)
-		{
-			switch (ErrorResult)
-			{
-				case (EInstallBundleReleaseResult::UserCancelledError):
-				{
-					return ReleaseResult_Cancelled;
-				}
-			
-				case (EInstallBundleReleaseResult::ManifestArchiveError):
-				{
-					return ReleaseResult_Generic;
-				}
-
-				default:
-				{
-					//Show generic error for anything missing but log an error
-					UE_LOG(LogGameFeatures, Error, TEXT("Missing error text for EInstallBundleReleaseResult %s"), LexToString(ErrorResult));
-					return ReleaseResult_Generic;
-				}
-			}
-		}
-	};
-
 	static bool bRealtimeMode = false;
 
 	class FRealtimeMode : public TSharedFromThis<FRealtimeMode>
@@ -1119,7 +1028,7 @@ struct FBaseDataReleaseGameFeaturePluginState : public FGameFeaturePluginState
 		{
 			const FStringView ShortUrl = StateProperties.PluginIdentifier.GetIdentifyingString();
 			ensureMsgf(false, TEXT("Unable to enqueue uninstall for the PluginURL(%.*s) because failed to resolve install bundles!"), ShortUrl.Len(), ShortUrl.GetData());
-			Result = GetErrorResult(TEXT("BundleManager.Begin."), TEXT("Resolve_Failed"), UE::GameFeatures::CommonErrorCodes::ReleaseResult_Generic);
+			Result = GetErrorResult(TEXT("BundleManager.Begin."), TEXT("Resolve_Failed"), UE::GameFeatures::CommonErrorCodes::GetGenericReleaseResult());
 
 			return;
 		}
@@ -1386,7 +1295,7 @@ struct FGameFeaturePluginState_Downloading : public FGameFeaturePluginState
 		{
 			const FStringView ShortUrl = StateProperties.PluginIdentifier.GetIdentifyingString();
 			ensureMsgf(false, TEXT("Unable to enqueue download for the PluginURL(%.*s) because failed to resolve install bundles!"), ShortUrl.Len(), ShortUrl.GetData());
-			Result = GetErrorResult(TEXT("BundleManager.GotState."), TEXT("Resolve_Failed"), UE::GameFeatures::CommonErrorCodes::Generic_ConnectionError);
+			Result = GetErrorResult(TEXT("BundleManager.GotState."), TEXT("Resolve_Failed"), UE::GameFeatures::CommonErrorCodes::GetGenericConnectionError());
 			
 			UpdateStateMachineImmediate();
 			return;
@@ -1766,7 +1675,7 @@ struct FGameFeaturePluginState_Unmounting : public FGameFeaturePluginState
 		{
 			const FStringView ShortUrl = StateProperties.PluginIdentifier.GetIdentifyingString();
 			ensureMsgf(false, TEXT("Unable to enqueue unmount for the PluginURL(%.*s) because failed to resolve install bundles!"), ShortUrl.Len(), ShortUrl.GetData());
-			Result = GetErrorResult(TEXT("BundleManager.Begin."), TEXT("Cannot_Resolve"), UE::GameFeatures::CommonErrorCodes::Generic_ConnectionError);
+			Result = GetErrorResult(TEXT("BundleManager.Begin."), TEXT("Cannot_Resolve"), UE::GameFeatures::CommonErrorCodes::GetGenericConnectionError());
 			return;
 		}
 
@@ -3593,7 +3502,7 @@ FString FInstallBundlePluginProtocolMetaData::ToString() const
 	return ReturnedString;
 }
 
-bool FInstallBundlePluginProtocolMetaData::FromString(const FString& URLString, FInstallBundlePluginProtocolMetaData& Metadata)
+bool FInstallBundlePluginProtocolMetaData::FromString(FStringView URLString, FInstallBundlePluginProtocolMetaData& Metadata)
 {
 	bool bParseSuccess = true;
 
@@ -3604,7 +3513,7 @@ bool FInstallBundlePluginProtocolMetaData::FromString(const FString& URLString,
 
 	if (URLOptions.Num() > 0)
 	{
-		ensureAlwaysMsgf(URLOptions[0].StartsWith(UE::GameFeatures::GameFeaturePluginProtocolPrefix(EGameFeaturePluginProtocol::InstallBundle)), TEXT("Unexpected URL Format! Expected Protocol and uplugin information at the beginning of the URL %s"), *URLString);
+		ensureAlwaysMsgf(URLOptions[0].StartsWith(UE::GameFeatures::GameFeaturePluginProtocolPrefix(EGameFeaturePluginProtocol::InstallBundle)), TEXT("Unexpected URL Format! Expected Protocol and uplugin information at the beginning of the URL %.*s"), URLString.Len(), URLString.GetData());
 
 		//Parse through our URLOptions. Start at ParsingIndex 1 as option 0 should always be the .uplugin path that doesn't contain metadata
 		for (int ParsingIndex = 1; ParsingIndex < URLOptions.Num(); ++ParsingIndex)
@@ -3626,7 +3535,7 @@ bool FInstallBundlePluginProtocolMetaData::FromString(const FString& URLString,
 						if (OptionStrings.Num() != 2)
 						{
 							bParseSuccess = false;
-							UE_LOG(LogGameFeatures, Error, TEXT("Error parsing InstallBundle protocol options URL %s . No Valid Bundle List Found!"), *URLString);
+							UE_LOG(LogGameFeatures, Error, TEXT("Error parsing InstallBundle protocol options URL %.*s. No Valid Bundle List Found!"), URLString.Len(), URLString.GetData());
 						}
 						else
 						{
@@ -3652,7 +3561,7 @@ bool FInstallBundlePluginProtocolMetaData::FromString(const FString& URLString,
 					default:
 					{
 						bParseSuccess = false;
-						UE_LOG(LogGameFeatures, Error, TEXT("Error parsing InstallBundle protocol options for URL %s . Unknown Option %.*s"), *URLString, URLOption.Len(), URLOption.GetData());
+						UE_LOG(LogGameFeatures, Error, TEXT("Error parsing InstallBundle protocol options for URL %.*s. Unknown Option %.*s"), URLString.Len(), URLString.GetData(), URLOption.Len(), URLOption.GetData());
 
 						break;
 					}
@@ -3664,7 +3573,7 @@ bool FInstallBundlePluginProtocolMetaData::FromString(const FString& URLString,
 	if (Metadata.InstallBundles.Num() == 0)
 	{
 		bParseSuccess = false;
-		UE_LOG(LogGameFeatures, Error, TEXT("Error parsing InstallBundle protocol options URL %s . No Bundle List Found!"), *URLString);
+		UE_LOG(LogGameFeatures, Error, TEXT("Error parsing InstallBundle protocol options URL %.*s. No Bundle List Found!"), URLString.Len(), URLString.GetData());
 	}
 
 	static_assert(static_cast<uint8>(EGameFeatureInstallBundleProtocolOptions::Count) == 1, "Update this function to handle the newly added EGameFeatureInstallBundleProtocolOptions value!");
diff --git a/Engine/Plugins/Experimental/GameFeatures/Source/GameFeatures/Private/GameFeaturePluginStateMachine.h b/Engine/Plugins/Experimental/GameFeatures/Source/GameFeatures/Private/GameFeaturePluginStateMachine.h
--- a/Engine/Plugins/Experimental/GameFeatures/Source/GameFeatures/Private/GameFeaturePluginStateMachine.h
+++ b/Engine/Plugins/Experimental/GameFeatures/Source/GameFeatures/Private/GameFeaturePluginStateMachine.h
@@ -12,8 +12,8 @@ class UGameFeatureData;
 class UGameFrameworkComponentManager;
 class UGameFeaturePluginStateMachine;
 struct FComponentRequestHandle;
-enum class EInstallBundleResult : int;
-enum class EInstallBundleReleaseResult;
+enum class EInstallBundleResult : uint32;
+enum class EInstallBundleReleaseResult : uint32;
 
 /*
 *************** GameFeaturePlugin state machine graph ***************
@@ -234,7 +234,7 @@ struct FInstallBundlePluginProtocolMetaData
 
 	/** Functions to convert to/from the URL FString representation of this metadata **/
 	FString ToString() const;
-	static bool FromString(const FString& URLString, FInstallBundlePluginProtocolMetaData& OutMetadata);
+	static bool FromString(FStringView URLString, FInstallBundlePluginProtocolMetaData& OutMetadata);
 };
 
 struct FGameFeatureProtocolMetadata : public TUnion<FInstallBundlePluginProtocolMetaData, FNull>
diff --git a/Engine/Plugins/Experimental/GameFeatures/Source/GameFeatures/Private/GameFeaturesSubsystem.cpp b/Engine/Plugins/Experimental/GameFeatures/Source/GameFeatures/Private/GameFeaturesSubsystem.cpp
--- a/Engine/Plugins/Experimental/GameFeatures/Source/GameFeatures/Private/GameFeaturesSubsystem.cpp
+++ b/Engine/Plugins/Experimental/GameFeatures/Source/GameFeatures/Private/GameFeaturesSubsystem.cpp
@@ -4,6 +4,7 @@
 #include "Algo/TopologicalSort.h"
 #include "AssetRegistry/AssetData.h"
 #include "AssetRegistry/IAssetRegistry.h"
+#include "BundlePrereqCombinedStatusHelper.h"
 #include "Dom/JsonValue.h"
 #include "GameFeaturesSubsystemSettings.h"
 #include "GameFeaturesProjectPolicies.h"
@@ -12,6 +13,7 @@
 #include "GameFeatureStateChangeObserver.h"
 #include "GameplayTagsManager.h"
 #include "Interfaces/IPluginManager.h"
+#include "Logging/StructuredLog.h"
 #include "Misc/App.h"
 #include "Misc/PackageName.h"
 #include "Misc/Paths.h"
@@ -1956,6 +1958,317 @@ void UGameFeaturesSubsystem::PruneCachedGameFeaturePluginDetails(const FString&
 	CachedPluginDetailsByFilename.Remove(PluginDescriptorFilename);
 }
 
+struct FGameFeaturePluginPredownloadContext : public FGameFeaturePluginPredownloadHandle
+{
+	const FStringView PredownloadErrorNamespace = TEXTVIEW("GameFeaturePlugin.Predownload.");
+	TMap<FGameFeaturePluginIdentifier, FInstallBundlePluginProtocolMetaData> GFPs;
+	TArray<FName> PendingBundleDownloads;
+	TOptional<FInstallBundleCombinedProgressTracker> ProgressTracker;
+	TUniqueFunction<void(const UE::GameFeatures::FResult&)> OnComplete;
+	TUniqueFunction<void(float)> OnProgress;
+	UE::GameFeatures::FResult Result = MakeValue();
+	float Progress = 0.0f;
+	bool bIsComplete = false;
+	bool bCanceled = false;
+
+	virtual ~FGameFeaturePluginPredownloadContext() override
+	{
+		Cleanup();
+	}
+
+	virtual bool IsComplete() const override
+	{
+		return bIsComplete;
+	}
+
+	virtual const UE::GameFeatures::FResult& GetResult() const override
+	{
+		return Result;
+	}
+
+	virtual float GetProgress() const override
+	{
+		return Progress;
+	}
+
+	virtual void Cancel() override
+	{
+		bCanceled = true;
+
+		if (PendingBundleDownloads.Num() > 0)
+		{
+			TSharedPtr<IInstallBundleManager> BundleManager = IInstallBundleManager::GetPlatformInstallBundleManager();
+			if (BundleManager)
+			{
+				BundleManager->CancelUpdateContent(PendingBundleDownloads);
+			}
+		}
+	}
+
+	void Cleanup()
+	{
+		ProgressTracker.Reset();
+		IInstallBundleManager::InstallBundleCompleteDelegate.RemoveAll(this);
+		IInstallBundleManager::PausedBundleDelegate.RemoveAll(this);
+	}
+
+	void SetComplete()
+	{
+		bIsComplete = true;
+		if (OnComplete)
+		{
+			OnComplete(Result);
+		}
+	}
+
+	void SetComplete(UE::GameFeatures::FResult&& InResult)
+	{
+		Result = MoveTemp(InResult);
+		bIsComplete = true;
+		if (OnComplete)
+		{
+			OnComplete(Result);
+		}
+	}
+
+	void SetCompleteCanceled()
+	{
+		Result = MakeError(FString::Printf(TEXT("%.*s%s"),
+			PredownloadErrorNamespace.Len(), PredownloadErrorNamespace.GetData(),
+			TEXT("Canceled")));
+		bIsComplete = true;
+		if (OnComplete)
+		{
+			OnComplete(Result);
+		}
+	}
+
+	void Start(TConstArrayView<FString> PluginURLs)
+	{
+		if (bCanceled)
+		{
+			SetCompleteCanceled();
+			return;
+		}
+
+		for (const FString& URL : PluginURLs)
+		{
+			if (UGameFeaturesSubsystem::GetPluginURLProtocol(URL) != EGameFeaturePluginProtocol::InstallBundle)
+			{
+				// Only support install bundle protocol for downloading right now
+				continue;
+			}
+
+			FInstallBundlePluginProtocolMetaData InstallBundleOptions;
+			const bool bParsedBundles = FInstallBundlePluginProtocolMetaData::FromString(URL, InstallBundleOptions);
+			if (!bParsedBundles)
+			{
+				UE_LOGFMT(LogGameFeatures, Error, "GFP Predownload failed to parse URL {URL}", ("URL", URL));
+				UE::GameFeatures::FResult ErrorResult = MakeError(FString::Printf(TEXT("%.*s%s"),
+					PredownloadErrorNamespace.Len(), PredownloadErrorNamespace.GetData(),
+					TEXT("BadUrl")));
+				SetComplete(MoveTemp(ErrorResult));
+				return;
+			}
+
+			GFPs.Emplace(URL, MoveTemp(InstallBundleOptions));
+		}
+
+		if (GFPs.Num() == 0)
+		{
+			SetComplete(MakeValue());
+			return;
+		}
+
+		UGameFeaturesSubsystem& GFPSubSys = UGameFeaturesSubsystem::Get();
+
+		TArray<FName> BundlesToInstall;
+		for (const TPair<FGameFeaturePluginIdentifier, FInstallBundlePluginProtocolMetaData>& Pair : GFPs)
+		{
+			UGameFeaturePluginStateMachine* Machine = GFPSubSys.FindGameFeaturePluginStateMachine(Pair.Key);
+			if (Machine && Machine->GetDestination() < EGameFeaturePluginState::Installed)
+			{
+				// Existing machine exists and wants to be uninstalled, can't precache
+				UE_LOGFMT(LogGameFeatures, Error, "GFP Predownload failed because a GFP is unloading, GFP: {GFP}", ("GFP", Machine->GetPluginName()));
+				UE::GameFeatures::FResult ErrorResult = MakeError(FString::Printf(TEXT("%.*s%s"),
+					PredownloadErrorNamespace.Len(), PredownloadErrorNamespace.GetData(),
+					TEXT("GFPUnloading")));
+				SetComplete(MoveTemp(ErrorResult));
+				return;
+			}
+
+			BundlesToInstall.Append(Pair.Value.InstallBundles);
+		}
+
+		TSharedPtr<IInstallBundleManager> BundleManager = IInstallBundleManager::GetPlatformInstallBundleManager();
+		if (!BundleManager)
+		{
+			UE_LOGFMT(LogGameFeatures, Error, "GFP Predownload failed, no Install Bundle Manager found.");
+			UE::GameFeatures::FResult ErrorResult = MakeError(FString::Printf(TEXT("%.*s%s"),
+				PredownloadErrorNamespace.Len(), PredownloadErrorNamespace.GetData(),
+				TEXT("BundleManager_Null")));
+			SetComplete(MoveTemp(ErrorResult));
+			return;
+		}
+
+		BundleManager->GetContentState(BundlesToInstall, EInstallBundleGetContentStateFlags::None, false,
+			FInstallBundleGetContentStateDelegate::CreateLambda([Context = SharedThis(this)](FInstallBundleCombinedContentState BundleContentState)
+			{ Context->OnGotContentState(MoveTemp(BundleContentState)); }));
+	}
+
+	void OnGotContentState(FInstallBundleCombinedContentState BundleContentState)
+	{
+		TSharedPtr<IInstallBundleManager> BundleManager = IInstallBundleManager::GetPlatformInstallBundleManager();
+
+		if (BundleContentState.GetAllBundlesHaveState(EInstallBundleInstallState::UpToDate))
+		{
+			SetComplete(MakeValue());
+			return;
+		}
+
+		if (bCanceled)
+		{
+			SetCompleteCanceled();
+			return;
+		}
+
+		TArray<FName> BundlesToInstall;
+		for (const TPair<FGameFeaturePluginIdentifier, FInstallBundlePluginProtocolMetaData>& Pair : GFPs)
+		{
+			BundlesToInstall.Append(Pair.Value.InstallBundles);
+		}
+
+		EInstallBundleRequestFlags InstallFlags = EInstallBundleRequestFlags::Defaults | EInstallBundleRequestFlags::SkipMount;
+		TValueOrError<FInstallBundleRequestInfo, EInstallBundleResult> MaybeRequestInfo = BundleManager->RequestUpdateContent(BundlesToInstall, InstallFlags);
+
+		if (MaybeRequestInfo.HasError())
+		{
+			UE_LOGFMT(LogGameFeatures, Error, "GFP Predownload failed to request content, Error: {Error}", ("Error", LexToString(MaybeRequestInfo.GetError())));
+			UE::GameFeatures::FResult ErrorResult = MakeError(FString::Printf(TEXT("%.*s%s"),
+				PredownloadErrorNamespace.Len(), PredownloadErrorNamespace.GetData(),
+				LexToString(MaybeRequestInfo.GetError())));
+			SetComplete(MoveTemp(ErrorResult));
+			return;
+		}
+
+		FInstallBundleRequestInfo RequestInfo = MaybeRequestInfo.StealValue();
+		if (RequestInfo.BundlesEnqueued.Num() == 0)
+		{
+			SetComplete(MakeValue());
+			return;
+		}
+
+		PendingBundleDownloads = MoveTemp(RequestInfo.BundlesEnqueued);
+
+		ProgressTracker.Emplace(true, [this](const FInstallBundleCombinedProgressTracker::FCombinedProgress& InProgress)
+		{
+			Progress = InProgress.ProgressPercent;
+			if (OnProgress)
+			{
+				OnProgress(InProgress.ProgressPercent);
+			}
+		});
+		ProgressTracker->SetBundlesToTrackFromContentState(BundleContentState, PendingBundleDownloads);
+
+		IInstallBundleManager::InstallBundleCompleteDelegate.AddLambda(
+			[Context = SharedThis(this)](FInstallBundleRequestResultInfo BundleResult)
+			{ Context->OnInstallBundleCompleted(MoveTemp(BundleResult)); });
+		// TODO: handle pause?  Just cancel? This should only be relevent for cell connections
+		// IInstallBundleManager::PausedBundleDelegate.AddRaw(this, &FGameFeaturePluginState_Downloading::OnInstallBundlePaused);
+	}
+
+	void OnInstallBundleCompleted(FInstallBundleRequestResultInfo BundleResult)
+	{
+		if (!PendingBundleDownloads.Contains(BundleResult.BundleName))
+		{
+			return;
+		}
+
+		PendingBundleDownloads.Remove(BundleResult.BundleName);
+
+		if (Result.HasValue() && BundleResult.Result != EInstallBundleResult::OK)
+		{
+			if (BundleResult.OptionalErrorCode.IsEmpty())
+			{
+				UE_LOGFMT(LogGameFeatures, Error, "GFP Predownload failed to install {Bundle}, Error: {Error}",
+					("Bundle", BundleResult.BundleName), ("Error", LexToString(BundleResult.Result)));
+			}
+			else
+			{
+				UE_LOGFMT(LogGameFeatures, Error, "GFP Predownload failed to install {Bundle}, Error: {Error}",
+					("Bundle", BundleResult.BundleName), ("Error", BundleResult.OptionalErrorCode));
+			}
+
+			//Use OptionalErrorCode and/or OptionalErrorText if available
+			const FString ErrorCodeEnding = (BundleResult.OptionalErrorCode.IsEmpty()) ? LexToString(BundleResult.Result) : BundleResult.OptionalErrorCode;
+			FText ErrorText = BundleResult.OptionalErrorCode.IsEmpty() ? UE::GameFeatures::CommonErrorCodes::GetErrorTextForBundleResult(BundleResult.Result) : BundleResult.OptionalErrorText;
+			Result = UE::GameFeatures::FResult(
+				MakeError(FString::Printf(TEXT("%.*s%s"), PredownloadErrorNamespace.Len(), PredownloadErrorNamespace.GetData(), *ErrorCodeEnding)),
+				MoveTemp(ErrorText)
+			);
+
+			// Cancel remaining downloads
+			TSharedPtr<IInstallBundleManager> BundleManager = IInstallBundleManager::GetPlatformInstallBundleManager();
+			BundleManager->CancelUpdateContent(PendingBundleDownloads);
+		}
+
+		if (PendingBundleDownloads.Num() > 0)
+		{
+			return;
+		}
+
+		// Delay call to ReleaseBundlesIfPossible. We don't want to release them from within the complete callback.
+		FTSTicker::GetCoreTicker().AddTicker(
+			FTickerDelegate::CreateLambda([Context = SharedThis(this)](float)
+			{
+				Context->ReleaseBundlesIfPossible();
+
+				// Done
+				Context->SetComplete();
+
+				Context->Cleanup();
+				return false;
+			})
+		);
+	}
+
+	// Predownload shouldn't pin any cached bundles so release them now
+	void ReleaseBundlesIfPossible()
+	{
+		UGameFeaturesSubsystem& GFPSubSys = UGameFeaturesSubsystem::Get();
+		TSharedPtr<IInstallBundleManager> BundleManager = IInstallBundleManager::GetPlatformInstallBundleManager();
+
+		TArray<FName> ReleaseList;
+		TArray<FName> KeepList;
+		for (const TPair<FGameFeaturePluginIdentifier, FInstallBundlePluginProtocolMetaData>& Pair : GFPs)
+		{
+			UGameFeaturePluginStateMachine* Machine = GFPSubSys.FindGameFeaturePluginStateMachine(Pair.Key);
+			if (Machine &&
+				Machine->GetCurrentState() > EGameFeaturePluginState::StatusKnown &&
+				Machine->GetCurrentState() != EGameFeaturePluginState::Releasing)
+			{
+				// A machine is using the bundles, don't release
+				KeepList.Append(Pair.Value.InstallBundles);
+				continue;
+			}
+
+			ReleaseList.Append(Pair.Value.InstallBundles);
+		}
+
+		BundleManager->RequestReleaseContent(ReleaseList, EInstallBundleReleaseRequestFlags::None, KeepList);
+	}
+};
+
+TSharedRef<FGameFeaturePluginPredownloadHandle> UGameFeaturesSubsystem::PredownloadGameFeaturePlugins(TConstArrayView<FString> PluginURLs, TUniqueFunction<void(const UE::GameFeatures::FResult&)> OnComplete /*= nullptr*/, TUniqueFunction<void(float)> OnProgress /*= nullptr*/)
+{
+	TSharedRef<FGameFeaturePluginPredownloadContext> Context = MakeShared<FGameFeaturePluginPredownloadContext>();
+	Context->OnComplete = MoveTemp(OnComplete);
+	Context->OnProgress = MoveTemp(OnProgress);
+	Context->Start(PluginURLs);
+
+	return Context;
+}
+
 UGameFeaturePluginStateMachine* UGameFeaturesSubsystem::FindGameFeaturePluginStateMachineByPluginName(const FString& PluginName) const
 {
 	for (auto StateMachineIt = GameFeaturePluginStateMachines.CreateConstIterator(); StateMachineIt; ++StateMachineIt)
@@ -2492,6 +2805,7 @@ void UGameFeaturesSubsystem::FilterInactivePluginAssets(TArray<FAssetData>& Asse
 		return !GameFeaturesSubsystem::IsContentWithinActivePlugin(Asset.GetObjectPathString(), ActivePluginNames);
 	});
 }
+
 EBuiltInAutoState UGameFeaturesSubsystem::DetermineBuiltInInitialFeatureState(TSharedPtr<FJsonObject> Descriptor, const FString& ErrorContext)
 {
 	EBuiltInAutoState InitialState = EBuiltInAutoState::Invalid;
diff --git a/Engine/Plugins/Experimental/GameFeatures/Source/GameFeatures/Public/GameFeaturePluginOperationResult.h b/Engine/Plugins/Experimental/GameFeatures/Source/GameFeatures/Public/GameFeaturePluginOperationResult.h
--- a/Engine/Plugins/Experimental/GameFeatures/Source/GameFeatures/Public/GameFeaturePluginOperationResult.h
+++ b/Engine/Plugins/Experimental/GameFeatures/Source/GameFeatures/Public/GameFeaturePluginOperationResult.h
@@ -5,6 +5,9 @@
 #include "Internationalization/Text.h"
 #include "Templates/ValueOrError.h"
 
+enum class EInstallBundleResult : uint32;
+enum class EInstallBundleReleaseResult : uint32;
+
 namespace UE::GameFeatures
 {
 	//Type used to determine if our FResult is actually an error or success (and hold the error code)
@@ -57,6 +60,17 @@ namespace UE::GameFeatures
 	};
 
 	GAMEFEATURES_API FString ToString(const FResult& Result);
+
+	namespace CommonErrorCodes
+	{
+		const FText& GetErrorTextForBundleResult(EInstallBundleResult ErrorResult);
+		const FText& GetErrorTextForReleaseResult(EInstallBundleReleaseResult ErrorResult);
+
+		const FText& GetGenericFatalError();
+		const FText& GetGenericConnectionError();
+		const FText& GetGenericMountError();
+		const FText& GetGenericReleaseResult();
+	}
 }
 
 #if UE_ENABLE_INCLUDE_ORDER_DEPRECATED_IN_5_2
diff --git a/Engine/Plugins/Experimental/GameFeatures/Source/GameFeatures/Public/GameFeaturesSubsystem.h b/Engine/Plugins/Experimental/GameFeatures/Source/GameFeatures/Public/GameFeaturesSubsystem.h
--- a/Engine/Plugins/Experimental/GameFeatures/Source/GameFeatures/Public/GameFeaturesSubsystem.h
+++ b/Engine/Plugins/Experimental/GameFeatures/Source/GameFeatures/Public/GameFeaturesSubsystem.h
@@ -260,6 +260,15 @@ struct FBuiltInGameFeaturePluginBehaviorOptions
 	bool bForceSyncLoading = false;
 };
 
+struct FGameFeaturePluginPredownloadHandle : public TSharedFromThis<FGameFeaturePluginPredownloadHandle>
+{
+	virtual ~FGameFeaturePluginPredownloadHandle() {}
+	virtual bool IsComplete() const = 0;
+	virtual const UE::GameFeatures::FResult& GetResult() const = 0;
+	virtual float GetProgress() const = 0;
+	virtual void Cancel() = 0;
+};
+
 /** Struct used to transform a GameFeaturePlugin URL into something that can uniquely identify the GameFeaturePlugin
     without including any transient data being passed in through the URL */
 USTRUCT()
@@ -573,6 +582,13 @@ public:
 	/** Gets relevant properties out of a uplugin file if it's installed */
 	bool GetGameFeaturePluginDetails(FString PluginURL, struct FGameFeaturePluginDetails& OutPluginDetails) const;
 
+	/** 
+	 * Pre-install any required game feature data, which can be useful for larger payloads. 
+	 * This does not instantiate any GFP although it is safe to do so before this finishes. 
+	*/
+	TSharedRef<FGameFeaturePluginPredownloadHandle> PredownloadGameFeaturePlugins(TConstArrayView<FString> PluginURLs, TUniqueFunction<void(const UE::GameFeatures::FResult&)> OnComplete = nullptr, TUniqueFunction<void(float)> OnProgress = nullptr);
+	friend struct FGameFeaturePluginPredownloadContext;
+
 	/** Determine the initial feature state for a built-in plugin */
 	static EBuiltInAutoState DetermineBuiltInInitialFeatureState(TSharedPtr<FJsonObject> Descriptor, const FString& ErrorContext);
 
diff --git a/Engine/Source/Runtime/InstallBundleManager/Private/BundlePrereqCombinedStatusHelper.cpp b/Engine/Source/Runtime/InstallBundleManager/Private/BundlePrereqCombinedStatusHelper.cpp
--- a/Engine/Source/Runtime/InstallBundleManager/Private/BundlePrereqCombinedStatusHelper.cpp
+++ b/Engine/Source/Runtime/InstallBundleManager/Private/BundlePrereqCombinedStatusHelper.cpp
@@ -23,7 +23,8 @@ const TCHAR* LexToString(FInstallBundleCombinedProgressTracker::ECombinedBundleS
 	return Strings[InstallBundleUtil::CastToUnderlying(Status)];
 }
 
-FInstallBundleCombinedProgressTracker::FInstallBundleCombinedProgressTracker(bool bAutoTick /*= true*/)
+FInstallBundleCombinedProgressTracker::FInstallBundleCombinedProgressTracker(bool bAutoTick /*= true*/, TUniqueFunction<void(const FCombinedProgress&)> InOnTick /*= nullptr*/)
+	: OnTick(MoveTemp(InOnTick))
 {
 	SetupDelegates(bAutoTick);
 }
@@ -322,6 +323,11 @@ bool FInstallBundleCombinedProgressTracker::Tick(float dt)
 	UpdateBundleCache();
 	UpdateCombinedStatus();
 	
+	if (OnTick)
+	{
+		OnTick(CurrentCombinedProgress);
+	}
+
 	//just always keep ticking
 	return true;
 }
diff --git a/Engine/Source/Runtime/InstallBundleManager/Private/InstallBundleCache.cpp b/Engine/Source/Runtime/InstallBundleManager/Private/InstallBundleCache.cpp
--- a/Engine/Source/Runtime/InstallBundleManager/Private/InstallBundleCache.cpp
+++ b/Engine/Source/Runtime/InstallBundleManager/Private/InstallBundleCache.cpp
@@ -266,6 +266,17 @@ bool FInstallBundleCache::Contains(EInstallBundleSourceType Source, FName Bundle
 	return false;
 }
 
+bool FInstallBundleCache::IsReserved(FName BundleName) const
+{
+	const FBundleCacheInfo* BundleInfo = CacheInfo.Find(BundleName);
+	if (BundleInfo)
+	{
+		return BundleInfo->State == ECacheState::Reserved;
+	}
+
+	return false;
+}
+
 bool FInstallBundleCache::Release(FName BundleName)
 {
 	CSV_SCOPED_TIMING_STAT(InstallBundleManager, FInstallBundleCache_Release);
diff --git a/Engine/Source/Runtime/InstallBundleManager/Public/BundlePrereqCombinedStatusHelper.h b/Engine/Source/Runtime/InstallBundleManager/Public/BundlePrereqCombinedStatusHelper.h
--- a/Engine/Source/Runtime/InstallBundleManager/Public/BundlePrereqCombinedStatusHelper.h
+++ b/Engine/Source/Runtime/InstallBundleManager/Public/BundlePrereqCombinedStatusHelper.h
@@ -50,7 +50,7 @@ public:
 	
 public:
 	/** bAutoTick Whehter this tracker should automatically Tick */
-	INSTALLBUNDLEMANAGER_API FInstallBundleCombinedProgressTracker(bool bAutoTick = true);
+	INSTALLBUNDLEMANAGER_API FInstallBundleCombinedProgressTracker(bool bAutoTick = true, TUniqueFunction<void(const FCombinedProgress&)> OnTick = nullptr);
 	INSTALLBUNDLEMANAGER_API ~FInstallBundleCombinedProgressTracker();
 	
 	INSTALLBUNDLEMANAGER_API FInstallBundleCombinedProgressTracker(const FInstallBundleCombinedProgressTracker& Other);
@@ -96,6 +96,7 @@ private:
 	
 	TWeakPtr<IInstallBundleManager> InstallBundleManager;
 	FTSTicker::FDelegateHandle TickHandle;
+	TUniqueFunction<void(const FCombinedProgress&)> OnTick;
 };
 
 INSTALLBUNDLEMANAGER_API const TCHAR* LexToString(FInstallBundleCombinedProgressTracker::ECombinedBundleStatus Status);
diff --git a/Engine/Source/Runtime/InstallBundleManager/Public/InstallBundleCache.h b/Engine/Source/Runtime/InstallBundleManager/Public/InstallBundleCache.h
--- a/Engine/Source/Runtime/InstallBundleManager/Public/InstallBundleCache.h
+++ b/Engine/Source/Runtime/InstallBundleManager/Public/InstallBundleCache.h
@@ -83,6 +83,8 @@ public:
 	INSTALLBUNDLEMANAGER_API bool Contains(FName BundleName) const;
 	INSTALLBUNDLEMANAGER_API bool Contains(EInstallBundleSourceType Source, FName BundleName) const;
 
+	INSTALLBUNDLEMANAGER_API bool IsReserved(FName BundleName) const;
+
 	// Called from bundle manager to make the files for this bundle eligible for eviction
 	INSTALLBUNDLEMANAGER_API bool Release(FName BundleName);
 
diff --git a/Engine/Source/Runtime/InstallBundleManager/Public/InstallBundleTypes.h b/Engine/Source/Runtime/InstallBundleManager/Public/InstallBundleTypes.h
--- a/Engine/Source/Runtime/InstallBundleManager/Public/InstallBundleTypes.h
+++ b/Engine/Source/Runtime/InstallBundleManager/Public/InstallBundleTypes.h
@@ -124,7 +124,7 @@ enum class EInstallBundleRequestInfoFlags : int32
 };
 ENUM_CLASS_FLAGS(EInstallBundleRequestInfoFlags);
 
-enum class EInstallBundleResult : int
+enum class EInstallBundleResult : uint32
 {
 	OK,
 	FailedPrereqRequiresLatestClient,
@@ -156,7 +156,7 @@ enum class EInstallBundleRequestFlags : uint32
 };
 ENUM_CLASS_FLAGS(EInstallBundleRequestFlags)
 
-enum class EInstallBundleReleaseResult
+enum class EInstallBundleReleaseResult : uint32
 {
 	OK,
 	ManifestArchiveError,
