Propogate bWasLoadedAsBuiltInGameFeaturePlugin when creating GFP dependencies. This ensures any dependencies skipped by plugin load filters still get the flag.
Add optional warnings/errors for when GFP depencies are forcibly created.  This can indicate a missing dependency or a dependency that was filtered out by the GFP policy.

diff --git a/Engine/Plugins/Experimental/GameFeatures/Source/GameFeatures/Private/GameFeaturePluginStateMachine.cpp b/Engine/Plugins/Experimental/GameFeatures/Source/GameFeatures/Private/GameFeaturePluginStateMachine.cpp
--- a/Engine/Plugins/Experimental/GameFeatures/Source/GameFeatures/Private/GameFeaturePluginStateMachine.cpp
+++ b/Engine/Plugins/Experimental/GameFeatures/Source/GameFeatures/Private/GameFeaturePluginStateMachine.cpp
@@ -2208,7 +2208,7 @@ struct FWaitingForDependenciesTransitionPolicy
 		UGameFeaturesSubsystem& GameFeaturesSubsystem = UGameFeaturesSubsystem::Get();
 
 		return GameFeaturesSubsystem.FindOrCreatePluginDependencyStateMachines(
-			*InStateProperties.PluginIdentifier.GetFullPluginURL(), InStateProperties.PluginInstalledFilename, InStateProperties.RecycleProtocolOptions(), OutDependencyMachines);
+			*InStateProperties.PluginIdentifier.GetFullPluginURL(), InStateProperties, OutDependencyMachines);
 	}
 
 	static FGameFeaturePluginStateRange GetDependencyStateRange()
diff --git a/Engine/Plugins/Experimental/GameFeatures/Source/GameFeatures/Private/GameFeaturesProjectPolicies.cpp b/Engine/Plugins/Experimental/GameFeatures/Source/GameFeatures/Private/GameFeaturesProjectPolicies.cpp
--- a/Engine/Plugins/Experimental/GameFeatures/Source/GameFeatures/Private/GameFeaturesProjectPolicies.cpp
+++ b/Engine/Plugins/Experimental/GameFeatures/Source/GameFeatures/Private/GameFeaturesProjectPolicies.cpp
@@ -16,6 +16,9 @@ void UDefaultGameFeaturesProjectPolicies::InitGameFeatureManager()
 		// By default, force all initially loaded plugins to synchronously load, this overrides the behavior of GameFeaturePlugin.AsyncLoad which will be used for later loads
 		OutOptions.bForceSyncLoading = true;
 
+		// By default, no plugins are filtered so we expect all built-in dependencies to be created before their parent GFPs
+		OutOptions.bLogWarningOnForcedDependencyCreation = true;
+
 		return true;
 	};
 
diff --git a/Engine/Plugins/Experimental/GameFeatures/Source/GameFeatures/Private/GameFeaturesSubsystem.cpp b/Engine/Plugins/Experimental/GameFeatures/Source/GameFeatures/Private/GameFeaturesSubsystem.cpp
--- a/Engine/Plugins/Experimental/GameFeatures/Source/GameFeatures/Private/GameFeaturesSubsystem.cpp
+++ b/Engine/Plugins/Experimental/GameFeatures/Source/GameFeatures/Private/GameFeaturesSubsystem.cpp
@@ -291,6 +291,10 @@ FSimpleDelegate FGameFeaturePostMountingContext::PauseUntilComplete(FString InPa
 FInstallBundlePluginProtocolOptions::FInstallBundlePluginProtocolOptions()
 	: InstallBundleFlags(EInstallBundleRequestFlags::Defaults)
 	, ReleaseInstallBundleFlags(EInstallBundleReleaseRequestFlags::None)
+	, bUninstallBeforeTerminate(false)
+	, bUserPauseDownload(false)
+	, bAllowIniLoading(false)
+	, bDoNotDownload(false)
 {}
 
 bool FInstallBundlePluginProtocolOptions::operator==(const FInstallBundlePluginProtocolOptions& Other) const
@@ -304,6 +308,31 @@ bool FInstallBundlePluginProtocolOptions::operator==(const FInstallBundlePluginP
 		bDoNotDownload == Other.bDoNotDownload;
 }
 
+// TODO : FGameFeatureProtocolOptions - C++20 will allow inline init for bitfields
+FGameFeatureProtocolOptions::FGameFeatureProtocolOptions()
+	: bForceSyncLoading(false)
+	, bLogWarningOnForcedDependencyCreation(false)
+	, bLogErrorOnForcedDependencyCreation(false)
+{ 
+	SetSubtype<FNull>(); 
+}
+
+FGameFeatureProtocolOptions::FGameFeatureProtocolOptions(const FInstallBundlePluginProtocolOptions& InOptions) 
+	: TUnion(InOptions) 
+	, bForceSyncLoading(false)
+	, bLogWarningOnForcedDependencyCreation(false)
+	, bLogErrorOnForcedDependencyCreation(false)
+{
+}
+
+FGameFeatureProtocolOptions::FGameFeatureProtocolOptions(FNull InOptions)
+	: bForceSyncLoading(false)
+	, bLogWarningOnForcedDependencyCreation(false)
+	, bLogErrorOnForcedDependencyCreation(false)
+{
+	SetSubtype<FNull>(InOptions);
+}
+
 void UGameFeaturesSubsystem::Initialize(FSubsystemCollectionBase& Collection)
 {
 	UE_LOG(LogGameFeatures, Log, TEXT("Initializing game features subsystem"));
@@ -1494,6 +1523,8 @@ void UGameFeaturesSubsystem::LoadBuiltInGameFeaturePlugin(const TSharedRef<IPlug
 			{
 				FGameFeatureProtocolOptions ProtocolOptions;
 				ProtocolOptions.bForceSyncLoading = BehaviorOptions.bForceSyncLoading;
+				ProtocolOptions.bLogWarningOnForcedDependencyCreation = BehaviorOptions.bLogWarningOnForcedDependencyCreation;
+				ProtocolOptions.bLogErrorOnForcedDependencyCreation = BehaviorOptions.bLogErrorOnForcedDependencyCreation;
 				UGameFeaturePluginStateMachine* StateMachine = FindOrCreateGameFeaturePluginStateMachine(PluginURL, ProtocolOptions);
 
 				EBuiltInAutoState InitialAutoState = (BehaviorOptions.AutoStateOverride != EBuiltInAutoState::Invalid) ? 
@@ -2313,19 +2344,24 @@ UGameFeaturePluginStateMachine* UGameFeaturesSubsystem::FindGameFeaturePluginSta
 
 // Note: ProtocolOptions is not defaulted here. Any API call that could create a state machine should allow the user to pass ProtocolOptions to initialize the machine.
 // It is acceptable that user passes null options. 
-UGameFeaturePluginStateMachine* UGameFeaturesSubsystem::FindOrCreateGameFeaturePluginStateMachine(const FString& PluginURL, const FGameFeatureProtocolOptions& ProtocolOptions)
+UGameFeaturePluginStateMachine* UGameFeaturesSubsystem::FindOrCreateGameFeaturePluginStateMachine(const FString& PluginURL, const FGameFeatureProtocolOptions& ProtocolOptions, bool* bOutFoundExisting /*= nullptr*/)
 {
 	TRACE_CPUPROFILER_EVENT_SCOPE(GFP_FindOrCreateStateMachine);
 	FGameFeaturePluginIdentifier PluginIdentifier(PluginURL);
 	TObjectPtr<UGameFeaturePluginStateMachine> const* ExistingStateMachine =
 		GameFeaturePluginStateMachines.FindByHash(GetTypeHash(PluginIdentifier.GetIdentifyingString()), PluginIdentifier.GetIdentifyingString());
-	if (ExistingStateMachine)
+
+	if (bOutFoundExisting)
 	{
-		EGameFeaturePluginProtocol ExpectedProtocol = (*ExistingStateMachine)->GetPluginIdentifier().GetPluginProtocol();
-		ensureAlwaysMsgf(ExpectedProtocol == PluginIdentifier.GetPluginProtocol(), TEXT("Expected protocol %s for %.*s"), UE::GameFeatures::GameFeaturePluginProtocolPrefix(ExpectedProtocol), PluginIdentifier.GetIdentifyingString().Len(), PluginIdentifier.GetIdentifyingString().GetData());
+		*bOutFoundExisting = !!ExistingStateMachine;
+	}
 
+	if (ExistingStateMachine)
+	{
 		// In this case, still return the existing machine, even if the protocol doesn't match. This function should never return null.
 		// There can only be one active instance of any machine.
+		EGameFeaturePluginProtocol ExpectedProtocol = (*ExistingStateMachine)->GetPluginIdentifier().GetPluginProtocol();
+		ensureAlwaysMsgf(ExpectedProtocol == PluginIdentifier.GetPluginProtocol(), TEXT("Expected protocol %s for %.*s"), UE::GameFeatures::GameFeaturePluginProtocolPrefix(ExpectedProtocol), PluginIdentifier.GetIdentifyingString().Len(), PluginIdentifier.GetIdentifyingString().GetData());
 
 		UE_LOG(LogGameFeatures, VeryVerbose, TEXT("Found GameFeaturePlugin StateMachine using Identifier:%.*s from PluginURL:%s"), PluginIdentifier.GetIdentifyingString().Len(), PluginIdentifier.GetIdentifyingString().GetData(), *PluginURL);
 		return *ExistingStateMachine;
@@ -2452,8 +2488,14 @@ void UGameFeaturesSubsystem::FinishTermination(UGameFeaturePluginStateMachine* M
 	TerminalGameFeaturePluginStateMachines.RemoveSwap(Machine);
 }
 
-bool UGameFeaturesSubsystem::FindOrCreatePluginDependencyStateMachines(const FString& PluginURL, const FString& PluginFilename, const FGameFeatureProtocolOptions& InDepProtocolOptions, TArray<UGameFeaturePluginStateMachine*>& OutDependencyMachines)
+bool UGameFeaturesSubsystem::FindOrCreatePluginDependencyStateMachines(const FString& PluginURL, const FGameFeaturePluginStateMachineProperties& InStateProperties, TArray<UGameFeaturePluginStateMachine*>& OutDependencyMachines)
 {
+	const FString& PluginFilename = InStateProperties.PluginInstalledFilename;
+	const FGameFeatureProtocolOptions InDepProtocolOptions = InStateProperties.RecycleProtocolOptions();
+
+	const bool bWarnOnDepCreation = InStateProperties.ProtocolOptions.bLogWarningOnForcedDependencyCreation;
+	const bool bErrorOnDepCreation = InStateProperties.ProtocolOptions.bLogErrorOnForcedDependencyCreation;
+
 	FGameFeaturePluginDetails Details;
 	if (GetGameFeaturePluginDetailsInternal(PluginFilename, Details))
 	{
@@ -2490,10 +2532,51 @@ bool UGameFeaturesSubsystem::FindOrCreatePluginDependencyStateMachines(const FSt
 			{
 				// Always propogate non-protocol specific flags
 				DepProtocolOptions.bForceSyncLoading = InDepProtocolOptions.bForceSyncLoading;
+				DepProtocolOptions.bLogWarningOnForcedDependencyCreation = InDepProtocolOptions.bLogWarningOnForcedDependencyCreation;
+				DepProtocolOptions.bLogErrorOnForcedDependencyCreation = InDepProtocolOptions.bLogErrorOnForcedDependencyCreation;
 			}
 
-			UGameFeaturePluginStateMachine* ResolvedDependency = FindOrCreateGameFeaturePluginStateMachine(DependencyURL, DepProtocolOptions);
+			bool bFoundExisting = false;
+			UGameFeaturePluginStateMachine* ResolvedDependency = FindOrCreateGameFeaturePluginStateMachine(DependencyURL, DepProtocolOptions, &bFoundExisting);
 			check(ResolvedDependency);
+
+			if (!bFoundExisting)
+			{
+				// Propogate bWasLoadedAsBuiltInGameFeaturePlugin
+				if (InStateProperties.bWasLoadedAsBuiltInGameFeaturePlugin)
+				{
+					ResolvedDependency->SetWasLoadedAsBuiltIn();
+				}
+
+				// Note: Given that LoadBuiltInGameFeaturePlugins does a topological sort, we don't expect to hit this path for built-ins
+				if (bWarnOnDepCreation)
+				{
+					if (InStateProperties.bWasLoadedAsBuiltInGameFeaturePlugin)
+					{
+						UE_LOGFMT(LogGameFeatures, Warning, "GFP dependency {Dep} was forcibly created by {Parent}, Game specific policies may be incorrectly filtering this dependency.",
+							("Dep", ResolvedDependency->GetPluginIdentifier().GetIdentifyingString()), ("Parent", InStateProperties.PluginIdentifier.GetIdentifyingString()));
+					}
+					else
+					{
+						UE_LOGFMT(LogGameFeatures, Warning, "GFP dependency {Dep} was unexpectedly forcibly created by {Parent}",
+							("Dep", ResolvedDependency->GetPluginIdentifier().GetIdentifyingString()), ("Parent", InStateProperties.PluginIdentifier.GetIdentifyingString()));
+					}
+				}
+				else if (bErrorOnDepCreation)
+				{
+					if (InStateProperties.bWasLoadedAsBuiltInGameFeaturePlugin)
+					{
+						UE_LOGFMT(LogGameFeatures, Error, "GFP dependency {Dep} was forcibly created by {Parent}, Game specific policies may be incorrectly filtering this dependency.",
+							("Dep", ResolvedDependency->GetPluginIdentifier().GetIdentifyingString()), ("Parent", InStateProperties.PluginIdentifier.GetIdentifyingString()));
+					}
+					else
+					{
+						UE_LOGFMT(LogGameFeatures, Error, "GFP dependency {Dep} was unexpectedly forcibly created by {Parent}",
+							("Dep", ResolvedDependency->GetPluginIdentifier().GetIdentifyingString()), ("Parent", InStateProperties.PluginIdentifier.GetIdentifyingString()));
+					}
+				}
+			}
+
 			OutDependencyMachines.Add(ResolvedDependency);
 		}
 
diff --git a/Engine/Plugins/Experimental/GameFeatures/Source/GameFeatures/Public/GameFeaturesSubsystem.h b/Engine/Plugins/Experimental/GameFeatures/Source/GameFeatures/Public/GameFeaturesSubsystem.h
--- a/Engine/Plugins/Experimental/GameFeatures/Source/GameFeatures/Public/GameFeaturesSubsystem.h
+++ b/Engine/Plugins/Experimental/GameFeatures/Source/GameFeatures/Public/GameFeaturesSubsystem.h
@@ -20,6 +20,7 @@ class IPlugin;
 class FJsonObject;
 struct FWorldContext;
 struct FGameFeaturePluginStateRange;
+struct FGameFeaturePluginStateMachineProperties;
 enum class EInstallBundleRequestFlags : uint32;
 enum class EInstallBundleReleaseRequestFlags : uint32;
 
@@ -257,7 +258,15 @@ struct FGameFeaturePluginDetails
 struct FBuiltInGameFeaturePluginBehaviorOptions
 {
 	EBuiltInAutoState AutoStateOverride = EBuiltInAutoState::Invalid;
+
+	/** Force this GFP to load synchronously even if async loading is allowed */
 	bool bForceSyncLoading = false;
+
+	/** Log Warning if loading this GFP forces dependencies to be created, useful for catching GFP load filtering bugs */
+	bool bLogWarningOnForcedDependencyCreation = false;
+
+	/** Log Error if loading this GFP forces dependencies to be created, useful for catching GFP load filtering bugs */
+	bool bLogErrorOnForcedDependencyCreation = false;
 };
 
 struct FGameFeaturePluginPredownloadHandle : public TSharedFromThis<FGameFeaturePluginPredownloadHandle>
@@ -345,28 +354,34 @@ struct GAMEFEATURES_API FInstallBundlePluginProtocolOptions
 	EInstallBundleReleaseRequestFlags ReleaseInstallBundleFlags;
 
 	/** If we want to attempt to uninstall InstallBundle data installed by this plugin before terminating */
-	bool bUninstallBeforeTerminate = false;
+	bool bUninstallBeforeTerminate : 1;
 
 	/** If we want to set the Downloading state to pause because of user interaction */
-	bool bUserPauseDownload = false;
+	bool bUserPauseDownload : 1;
 
 	/** Allow the GFP to load INI files, should only be allowed for trusted content */
-	bool bAllowIniLoading = false;
+	bool bAllowIniLoading : 1;
 
 	/** Disallows downloading, useful for conditionally loading content only if it's already been installed **/
-	bool bDoNotDownload = false;
+	bool bDoNotDownload : 1;
 
 	bool operator==(const FInstallBundlePluginProtocolOptions& Other) const;
 };
 
 struct FGameFeatureProtocolOptions : public TUnion<FInstallBundlePluginProtocolOptions, FNull>
 {
-	FGameFeatureProtocolOptions() { SetSubtype<FNull>(); }
-	explicit FGameFeatureProtocolOptions(const FInstallBundlePluginProtocolOptions& InOptions) : TUnion(InOptions) {}
-	explicit FGameFeatureProtocolOptions(FNull InOptions) { SetSubtype<FNull>(InOptions); }
+	GAMEFEATURES_API FGameFeatureProtocolOptions();
+	GAMEFEATURES_API explicit FGameFeatureProtocolOptions(const FInstallBundlePluginProtocolOptions& InOptions);
+	GAMEFEATURES_API explicit FGameFeatureProtocolOptions(FNull InOptions);
 
 	/** Force this GFP to load synchronously even if async loading is allowed */
-	bool bForceSyncLoading = false;
+	bool bForceSyncLoading : 1;
+
+	/** Log Warning if loading this GFP forces dependencies to be created, useful for catching GFP load filtering bugs */
+	bool bLogWarningOnForcedDependencyCreation : 1;
+
+	/** Log Error if loading this GFP forces dependencies to be created, useful for catching GFP load filtering bugs */
+	bool bLogErrorOnForcedDependencyCreation : 1;
 };
 
 /** The manager subsystem for game features */
@@ -665,7 +680,7 @@ private:
 	UGameFeaturePluginStateMachine* FindGameFeaturePluginStateMachine(const FGameFeaturePluginIdentifier& PluginIdentifier) const;
 
 	/** Gets the state machine associated with the specified URL, creates it if it doesnt exist */
-	UGameFeaturePluginStateMachine* FindOrCreateGameFeaturePluginStateMachine(const FString& PluginURL, const FGameFeatureProtocolOptions& ProtocolOptions);
+	UGameFeaturePluginStateMachine* FindOrCreateGameFeaturePluginStateMachine(const FString& PluginURL, const FGameFeatureProtocolOptions& ProtocolOptions, bool* bOutFoundExisting = nullptr);
 
 	/** Notification that a game feature has finished loading, and whether it was successful */
 	void LoadBuiltInGameFeaturePluginComplete(const UE::GameFeatures::FResult& Result, UGameFeaturePluginStateMachine* Machine, FGameFeaturePluginStateRange RequestedDestination);
@@ -685,7 +700,7 @@ private:
 	friend class UGameFeaturePluginStateMachine;
 
 	/** Handler for when a state machine requests its dependencies. Returns false if the dependencies could not be read */
-	bool FindOrCreatePluginDependencyStateMachines(const FString& PluginURL, const FString& PluginFilename, const FGameFeatureProtocolOptions& DepProtocolOptions, TArray<UGameFeaturePluginStateMachine*>& OutDependencyMachines);
+	bool FindOrCreatePluginDependencyStateMachines(const FString& PluginURL, const FGameFeaturePluginStateMachineProperties& InStateProperties, TArray<UGameFeaturePluginStateMachine*>& OutDependencyMachines);
 	template <typename> friend struct FTransitionDependenciesGameFeaturePluginState;
 	friend struct FWaitingForDependenciesTransitionPolicy;
 
