Added logging to content beacon for determining gfp load as built in states.

diff --git a/Engine/Plugins/Experimental/GameFeatures/Source/GameFeatures/Private/GameFeaturesSubsystem.cpp b/Engine/Plugins/Experimental/GameFeatures/Source/GameFeatures/Private/GameFeaturesSubsystem.cpp
--- a/Engine/Plugins/Experimental/GameFeatures/Source/GameFeatures/Private/GameFeaturesSubsystem.cpp
+++ b/Engine/Plugins/Experimental/GameFeatures/Source/GameFeatures/Private/GameFeaturesSubsystem.cpp
@@ -655,6 +655,18 @@ void UGameFeaturesSubsystem::RemoveGameFeatureFromAssetManager(const UGameFeatur
 	}
 }
 
+void UGameFeaturesSubsystem::ForEachGameFeature(TFunction<void(FGameFeatureInfo&&)>& Visitor) const
+{
+	for (auto StateMachineIt = GameFeaturePluginStateMachines.CreateConstIterator(); StateMachineIt; ++StateMachineIt)
+	{
+		if (UGameFeaturePluginStateMachine* GFSM = StateMachineIt.Value())
+		{	
+			FGameFeatureInfo GameFeatureInfo = { GFSM->GetPluginName(), GFSM->GetPluginURL(), GFSM->WasLoadedAsBuiltIn(), GFSM->GetCurrentState() };
+			Visitor(MoveTemp(GameFeatureInfo));
+		}
+	}
+}
+
 void UGameFeaturesSubsystem::AddObserver(UObject* Observer)
 {
 	//@TODO: GameFeaturePluginEnginePush: May want to warn if one is added after any game feature plugins are already initialized, or go to a CallOrRegister sort of pattern
diff --git a/Engine/Plugins/Experimental/GameFeatures/Source/GameFeatures/Public/GameFeaturesSubsystem.h b/Engine/Plugins/Experimental/GameFeatures/Source/GameFeatures/Public/GameFeaturesSubsystem.h
--- a/Engine/Plugins/Experimental/GameFeatures/Source/GameFeatures/Public/GameFeaturesSubsystem.h
+++ b/Engine/Plugins/Experimental/GameFeatures/Source/GameFeatures/Public/GameFeaturesSubsystem.h
@@ -384,6 +384,16 @@ struct FGameFeatureProtocolOptions : public TUnion<FInstallBundlePluginProtocolO
 	bool bLogErrorOnForcedDependencyCreation : 1;
 };
 
+// some important information about a gamefeature
+struct FGameFeatureInfo
+{
+	FString Name;
+	FString URL;
+	bool bLoadedAsBuiltIn;
+	EGameFeaturePluginState CurrentState;
+	EGameFeaturePluginState DestinationState;
+};
+
 /** The manager subsystem for game features */
 UCLASS()
 class GAMEFEATURES_API UGameFeaturesSubsystem : public UEngineSubsystem
@@ -406,6 +416,8 @@ public:
 	void AddObserver(UObject* Observer);
 	void RemoveObserver(UObject* Observer);
 
+	void ForEachGameFeature(TFunction<void(FGameFeatureInfo&&)>& Visitor) const;
+
 	/**
 	 * Calls the compile-time lambda on each active game feature data of the specified type
 	 * @param GameFeatureDataType       The kind of data required
