Fix disabled GFP depenencies being instantiated when their parent loads. This was caused by conflating whether a plugin was globally enabled with whether a plugin dependency reference was enabled.

diff --git a/Engine/Plugins/Experimental/GameFeatures/Source/GameFeatures/Private/GameFeaturesSubsystem.cpp b/Engine/Plugins/Experimental/GameFeatures/Source/GameFeatures/Private/GameFeaturesSubsystem.cpp
--- a/Engine/Plugins/Experimental/GameFeatures/Source/GameFeatures/Private/GameFeaturesSubsystem.cpp
+++ b/Engine/Plugins/Experimental/GameFeatures/Source/GameFeatures/Private/GameFeaturesSubsystem.cpp
@@ -1945,35 +1945,15 @@ bool UGameFeaturesSubsystem::GetGameFeaturePluginDetailsInternal(const FString&
 					ElementObject->TryGetStringField(NameField, DependencyName);
 					if (!DependencyName.IsEmpty())
 					{
-						//parse enabled element from IPlugin data and use that to filter if further parsing is needed
 						bool bElementEnabled = false;
-						if (TSharedPtr<IPlugin> DependencyPlugin = IPluginManager::Get().FindPlugin(DependencyName))
-						{
-							if (DependencyPlugin->GetType() == EPluginType::Engine)
-							{
-								UE_LOG(LogGameFeatures, VeryVerbose, TEXT("Skipping Dependency %s in %s because it is an engine plugin and thus can not be a GFP"), *DependencyName, *PluginDescriptorFilename);
-								continue;
-							}
-
-							bElementEnabled = DependencyPlugin->IsEnabled();							
-						}
-						// Plugin is not yet in PluginManager, so need to just go off JSon enabled value
-						else
-						{
-							ElementObject->TryGetBoolField(EnabledField, bElementEnabled);
-							if (bElementEnabled)
-							{
-								UE_LOG(LogGameFeatures, Verbose, TEXT("Plugin dependency %s marked enabled in %s but not found in PluginManager."), *DependencyName, *PluginDescriptorFilename);
-							}
-						}
-
+						ElementObject->TryGetBoolField(EnabledField, bElementEnabled);
 						if (bElementEnabled)
 						{
 							//Have to get Activate from JSON as it's unique to GFP and not in the PluginManager
 							bool bElementActivate = false;
 							ElementObject->TryGetBoolField(ActivateField, bElementActivate);
 
-							OutPluginDetails.PluginDependencies.Emplace(FGameFeaturePluginReferenceDetails(DependencyName, bElementActivate));
+							OutPluginDetails.PluginDependencies.Emplace(FGameFeaturePluginReferenceDetails(MoveTemp(DependencyName), bElementActivate));
 						}
 						else
 						{
