Fix hotfixed device profile fragments causing the entire fragment to load with broken configuration
 - Fix a bug where a hotfix would trigger a wipe of all non-hotfixed values in that section
 - Improve logging for this feature to get detailed logging when hotfixing
 - Reduce verbosity from Display to Log for the remaining logs

#tests Tested on a hotfixed MCP backend that changes are correctly pushed
#rnx

diff --git a/Engine/Plugins/Experimental/GameFeatures/Source/GameFeatures/Private/GameFeatureData.cpp b/Engine/Plugins/Experimental/GameFeatures/Source/GameFeatures/Private/GameFeatureData.cpp
--- a/Engine/Plugins/Experimental/GameFeatures/Source/GameFeatures/Private/GameFeatureData.cpp
+++ b/Engine/Plugins/Experimental/GameFeatures/Source/GameFeatures/Private/GameFeatureData.cpp
@@ -241,7 +241,7 @@ void UGameFeatureData::InitializeHierarchicalPluginIniFiles(const FString& Plugi
 				// Process new runtime device profile
 				if (bIsRuntimeDeviceProfileRule)
 				{
-					UE_LOG(LogGameFeatures, Display, TEXT("Game feature '%s' found runtime device profile rule %s"), *PluginName, *RuleName);
+					UE_LOG(LogGameFeatures, Log, TEXT("Game feature '%s' found runtime device profile rule %s"), *PluginName, *RuleName);
 
 					// Extract metadata
 					const FConfigValue* ParentProfileName = Section.Value.Find("ParentProfileName");
@@ -274,7 +274,7 @@ void UGameFeatureData::InitializeHierarchicalPluginIniFiles(const FString& Plugi
 								const FString FinalProfileName = Profile->GetName() + ProfileSuffix->GetValue();
 								if (ResultingProfiles.Contains(FinalProfileName))
 								{
-									UE_LOG(LogGameFeatures, Display, TEXT("Ignoring profile %s that has already been overriden as %s"), *Profile->GetName(), *FinalProfileName);
+									UE_LOG(LogGameFeatures, Log, TEXT("Ignoring profile %s that has already been overriden as %s"), *Profile->GetName(), *FinalProfileName);
 									continue;
 								}
 
@@ -282,7 +282,7 @@ void UGameFeatureData::InitializeHierarchicalPluginIniFiles(const FString& Plugi
 								RuntimeProfile.Add("DeviceType", PlatformName);
 								RuntimeProfile.Add("BaseProfileName", FConfigValue(Profile->GetName()));
 
-								UE_LOG(LogGameFeatures, Display, TEXT("Creating override for base profile %s"), *Profile->GetName());
+								UE_LOG(LogGameFeatures, Log, TEXT("Creating override for base profile %s"), *Profile->GetName());
 
 								// Inject the parent's matched fragments into the config, if any
 								if (Profile->GetName().Contains("MatchedFragments"))
@@ -298,7 +298,7 @@ void UGameFeatureData::InitializeHierarchicalPluginIniFiles(const FString& Plugi
 #endif
 
 									PlatformConfigSystem->GetArray(*MatchingRulesSectionName, *MatchingRulesArrayName, MatchingRulesArray, GDeviceProfilesIni);
-									UE_LOG(LogGameFeatures, Display, TEXT("Found %d fragment matching rules"), MatchingRulesArray.Num());
+									UE_LOG(LogGameFeatures, Log, TEXT("Found %d fragment matching rules"), MatchingRulesArray.Num());
 
 									for (const FString& Rule : MatchingRulesArray)
 									{
@@ -318,7 +318,7 @@ void UGameFeatureData::InitializeHierarchicalPluginIniFiles(const FString& Plugi
 									FString CVarKey, CVarValue;
 									if (CVar.Value.GetValue().Split(TEXT("="), &CVarKey, &CVarValue) && ShouldKeepCVar(CVarKey, CVarValue))
 									{
-										UE_LOG(LogGameFeatures, Verbose, TEXT(" Found CVar: %s=%s"), *CVarKey, *CVarValue);
+										UE_LOG(LogGameFeatures, Log, TEXT(" Found CVar: %s=%s"), *CVarKey, *CVarValue);
 										RuntimeProfile.Add(CVar.Key, FConfigValue(CVarKey + "=" + CVarValue));
 									}
 								}
@@ -329,8 +329,8 @@ void UGameFeatureData::InitializeHierarchicalPluginIniFiles(const FString& Plugi
 									FString CVarKey, CVarValue;
 									if (CVar.GetValue().Split(TEXT("="), &CVarKey, &CVarValue) && !PluginCVars.Contains(FName(*CVarKey)))
 									{
-										UE_LOG(LogGameFeatures, Verbose, TEXT(" Added CVar: %s=%s"), *CVarKey, *CVarValue);
-										RuntimeProfile.Add(FName(*CVarKey), FConfigValue(CVarKey + "=" + CVarValue));
+										UE_LOG(LogGameFeatures, Log, TEXT(" Added CVar: %s=%s"), *CVarKey, *CVarValue);
+										RuntimeProfile.Add("+CVars", FConfigValue(CVarKey + "=" + CVarValue));
 									}
 								}
 
@@ -350,31 +350,31 @@ void UGameFeatureData::InitializeHierarchicalPluginIniFiles(const FString& Plugi
 				// Hotfix device profile fragments
 				else if (bIsDeviceProfileFragment && HotfixCVars.Num() > 0)
 				{
-					UE_LOG(LogGameFeatures, Display, TEXT("Game feature '%s' found device profile fragment %s"), *PluginName, *RuleName);
+					UE_LOG(LogGameFeatures, Log, TEXT("Game feature '%s' found device profile fragment %s"), *PluginName, *RuleName);
 
 					// Update existing CVars
 					for (const auto& CVar : PluginCVars)
 					{
-						FString CVarKey, CVarValue;
-						if (CVar.Value.GetValue().Split(TEXT("="), &CVarKey, &CVarValue) && !ShouldKeepCVar(CVarKey, CVarValue))
+						FString PluginCVarKey, PluginCVarValue;
+						if (CVar.Value.GetValue().Split(TEXT("="), &PluginCVarKey, &PluginCVarValue) && !ShouldKeepCVar(PluginCVarKey, PluginCVarValue))
 						{
-							UE_LOG(LogGameFeatures, Verbose, TEXT(" Removed CVar: %s=%s"), *CVarKey, *CVarValue);
-							PluginConfig.RemoveKeyFromSection(*Section.Key, CVar.Key);
+							UE_LOG(LogGameFeatures, Log, TEXT(" Removed CVar: %s"), *CVar.Value.GetValue());
+							PluginConfig.RemoveFromSection(*Section.Key, CVar.Key, CVar.Value.GetValue());
 						}
 						else
 						{
-							UE_LOG(LogGameFeatures, Verbose, TEXT(" Found CVar: %s=%s"), *CVarKey, *CVarValue);
+							UE_LOG(LogGameFeatures, Log, TEXT(" Kept CVar: %s=%s"), *PluginCVarKey, *PluginCVarValue);
 						}
 					}
 
 					// Add new hotfix CVars
 					for (const auto& CVar : HotfixCVars)
 					{
-						FString CVarKey, CVarValue;
-						if (CVar.GetValue().Split(TEXT("="), &CVarKey, &CVarValue) && !PluginCVars.Contains(FName(*CVarKey)))
+						FString HotfixCVarKey, HotfixCVarValue;
+						if (CVar.GetValue().Split(TEXT("="), &HotfixCVarKey, &HotfixCVarValue) && !PluginCVars.Contains(FName(*HotfixCVarKey)))
 						{
-							UE_LOG(LogGameFeatures, Verbose, TEXT(" Added CVar: %s=%s"), *CVarKey, *CVarValue);
-							PluginConfig.AddToSection(*Section.Key, *CVarKey, CVarKey + "=" + CVarValue);
+							UE_LOG(LogGameFeatures, Log, TEXT(" Added CVar: %s=%s"), *HotfixCVarKey, *HotfixCVarValue);
+							PluginConfig.AddToSection(*Section.Key, "+CVars", HotfixCVarKey + "=" + HotfixCVarValue);
 						}
 					}
 				}
@@ -396,7 +396,7 @@ void UGameFeatureData::InitializeHierarchicalPluginIniFiles(const FString& Plugi
 				const FConfigValue* DeviceType = Section.Value.Find("DeviceType");
 				if (DeviceType)
 				{
-					UE_LOG(LogGameFeatures, Display, TEXT("Game feature '%s' adding new device profile %s"), *PluginName, *ProfileName);
+					UE_LOG(LogGameFeatures, Log, TEXT("Game feature '%s' adding new device profile %s"), *PluginName, *ProfileName);
 					DeviceProfileManager.CreateProfile(ProfileName, DeviceType->GetValue(), FString(), *PlatformName);
 				}
 			}
@@ -434,7 +434,7 @@ void UGameFeatureData::InitializeHierarchicalPluginIniFiles(const FString& Plugi
 		FConfigFile Config;
 		if (FConfigCacheIni::LoadExternalIniFile(Config, *PluginIniName, *EngineConfigDir, *ConfigDirectory, bIsBaseIniName, nullptr, bForceReloadFromDisk, bWriteDestIni) && (Config.Num() > 0))
 		{
-			UE_LOG(LogGameFeatures, Display, TEXT("Game feature '%s' loaded config file %s"), *PluginName, *PluginIniName);
+			UE_LOG(LogGameFeatures, Log, TEXT("Game feature '%s' loaded config file %s"), *PluginName, *PluginIniName);
 
 			// Need to get the in-memory config filename, the on disk one is likely not up to date
 			FString IniFile = GConfig->GetConfigFilename(*Ini.Name);
