Remove restriction on runtime device profiles that prevented empty rules from being applied, even if a hotfix exists for them
 - This allows us to create dedicated profiles that will be applied as a copy of their parent until new configurations are added through config or hotfix.
 - Also ensure device profile fragments are always logged, not just when a hotfix is applied to them.

#tests Confirmed the change allows empty profiles to become activated
#rnx
#localization none

diff --git a/Engine/Plugins/Experimental/GameFeatures/Source/GameFeatures/Private/GameFeatureData.cpp b/Engine/Plugins/Experimental/GameFeatures/Source/GameFeatures/Private/GameFeatureData.cpp
--- a/Engine/Plugins/Experimental/GameFeatures/Source/GameFeatures/Private/GameFeatureData.cpp
+++ b/Engine/Plugins/Experimental/GameFeatures/Source/GameFeatures/Private/GameFeatureData.cpp
@@ -247,7 +247,7 @@ void UGameFeatureData::InitializeHierarchicalPluginIniFiles(const FString& Plugi
 					const FConfigValue* ParentProfileName = Section.Value.Find("ParentProfileName");
 					const FConfigValue* ProfileSuffix = Section.Value.Find("ProfileSuffix");
 
-					if (ParentProfileName && ProfileSuffix && (FragmentIncludes.Num() > 0 || PluginCVars.Num() > 0))
+					if (ParentProfileName && ProfileSuffix)
 					{
 						// We need to load all candidate device profiles here or else we won't be able to create a child for them
 						TArray<FString> LoadableProfileNames = DeviceProfileManager.GetLoadableProfileNames(*PlatformName);
@@ -348,33 +348,36 @@ void UGameFeatureData::InitializeHierarchicalPluginIniFiles(const FString& Plugi
 				}
 
 				// Hotfix device profile fragments
-				else if (bIsDeviceProfileFragment && HotfixCVars.Num() > 0)
+				else if (bIsDeviceProfileFragment)
 				{
 					UE_LOG(LogGameFeatures, Log, TEXT("Game feature '%s' found device profile fragment %s"), *PluginName, *RuleName);
 
-					// Update existing CVars
-					for (const auto& CVar : PluginCVars)
+					if (HotfixCVars.Num() > 0)
 					{
-						FString PluginCVarKey, PluginCVarValue;
-						if (CVar.Value.GetValue().Split(TEXT("="), &PluginCVarKey, &PluginCVarValue) && !ShouldKeepCVar(PluginCVarKey, PluginCVarValue))
+						// Update existing CVars
+						for (const auto& CVar : PluginCVars)
 						{
-							UE_LOG(LogGameFeatures, Log, TEXT(" Removed CVar: %s"), *CVar.Value.GetValue());
-							PluginConfig.RemoveFromSection(*Section.Key, CVar.Key, CVar.Value.GetValue());
-						}
-						else
-						{
-							UE_LOG(LogGameFeatures, Log, TEXT(" Kept CVar: %s=%s"), *PluginCVarKey, *PluginCVarValue);
+							FString PluginCVarKey, PluginCVarValue;
+							if (CVar.Value.GetValue().Split(TEXT("="), &PluginCVarKey, &PluginCVarValue) && !ShouldKeepCVar(PluginCVarKey, PluginCVarValue))
+							{
+								UE_LOG(LogGameFeatures, Log, TEXT(" Removed CVar: %s"), *CVar.Value.GetValue());
+								PluginConfig.RemoveFromSection(*Section.Key, CVar.Key, CVar.Value.GetValue());
+							}
+							else
+							{
+								UE_LOG(LogGameFeatures, Log, TEXT(" Kept CVar: %s=%s"), *PluginCVarKey, *PluginCVarValue);
+							}
 						}
-					}
 
-					// Add new hotfix CVars
-					for (const auto& CVar : HotfixCVars)
-					{
-						FString HotfixCVarKey, HotfixCVarValue;
-						if (CVar.GetValue().Split(TEXT("="), &HotfixCVarKey, &HotfixCVarValue) && !PluginCVars.Contains(FName(*HotfixCVarKey)))
+						// Add new hotfix CVars
+						for (const auto& CVar : HotfixCVars)
 						{
-							UE_LOG(LogGameFeatures, Log, TEXT(" Added CVar: %s=%s"), *HotfixCVarKey, *HotfixCVarValue);
-							PluginConfig.AddToSection(*Section.Key, "+CVars", HotfixCVarKey + "=" + HotfixCVarValue);
+							FString HotfixCVarKey, HotfixCVarValue;
+							if (CVar.GetValue().Split(TEXT("="), &HotfixCVarKey, &HotfixCVarValue) && !PluginCVars.Contains(FName(*HotfixCVarKey)))
+							{
+								UE_LOG(LogGameFeatures, Log, TEXT(" Added CVar: %s=%s"), *HotfixCVarKey, *HotfixCVarValue);
+								PluginConfig.AddToSection(*Section.Key, "+CVars", HotfixCVarKey + "=" + HotfixCVarValue);
+							}
 						}
 					}
 				}
