#UE Console command ListGameFeaturePlugins now always alphasorts. The order without alphasorting is not useful since it came from an unordered TMap.

diff --git a/Engine/Plugins/Experimental/GameFeatures/Source/GameFeatures/Private/GameFeaturesSubsystem.cpp b/Engine/Plugins/Experimental/GameFeatures/Source/GameFeatures/Private/GameFeaturesSubsystem.cpp
--- a/Engine/Plugins/Experimental/GameFeatures/Source/GameFeatures/Private/GameFeaturesSubsystem.cpp
+++ b/Engine/Plugins/Experimental/GameFeatures/Source/GameFeatures/Private/GameFeaturesSubsystem.cpp
@@ -360,7 +360,7 @@ void UGameFeaturesSubsystem::Initialize(FSubsystemCollectionBase& Collection)
 
 	IConsoleManager::Get().RegisterConsoleCommand(
 		TEXT("ListGameFeaturePlugins"),
-		TEXT("Prints game features plugins and their current state to log. (options: [-activeonly] [-alphasort] [-csv])"),
+		TEXT("Prints game features plugins and their current state to log. (options: [-activeonly] [-csv])"),
 		FConsoleCommandWithWorldArgsAndOutputDeviceDelegate::CreateUObject(this, &ThisClass::ListGameFeaturePlugins),
 		ECVF_Default);
 
@@ -2655,7 +2655,6 @@ bool UGameFeaturesSubsystem::EnumeratePluginDependenciesWithShouldActivate(const
 
 void UGameFeaturesSubsystem::ListGameFeaturePlugins(const TArray<FString>& Args, UWorld* InWorld, FOutputDevice& Ar)
 {
-	const bool bAlphaSort = Args.ContainsByPredicate([](const FString& Arg) { return Arg.Compare(TEXT("-ALPHASORT"), ESearchCase::IgnoreCase) == 0; });
 	const bool bActiveOnly = Args.ContainsByPredicate([](const FString& Arg) { return Arg.Compare(TEXT("-ACTIVEONLY"), ESearchCase::IgnoreCase) == 0; });
 	const bool bCsv = Args.ContainsByPredicate([](const FString& Arg) { return Arg.Compare(TEXT("-CSV"), ESearchCase::IgnoreCase) == 0; });
 
@@ -2670,10 +2669,8 @@ void UGameFeaturesSubsystem::ListGameFeaturePlugins(const TArray<FString>& Args,
 	TArray<typename decltype(GameFeaturePluginStateMachines)::ValueType> StateMachines;
 	GameFeaturePluginStateMachines.GenerateValueArray(StateMachines);
 
-	if (bAlphaSort)
-	{
-		StateMachines.Sort([](const UGameFeaturePluginStateMachine& A, const UGameFeaturePluginStateMachine& B) { return A.GetGameFeatureName().Compare(B.GetGameFeatureName()) < 0; });
-	}
+	// Alphasort
+	StateMachines.Sort([](const UGameFeaturePluginStateMachine& A, const UGameFeaturePluginStateMachine& B) { return A.GetGameFeatureName().Compare(B.GetGameFeatureName()) < 0; });
 
 	int32 PluginCount = 0;
 	for (UGameFeaturePluginStateMachine* GFSM : StateMachines)
