Remove ensure that the GameFeatureData can't be null on unregister/unload in the editor because the asset can be force deleted
#rnx

diff --git a/Engine/Plugins/Experimental/GameFeatures/Source/GameFeatures/Private/GameFeaturesSubsystem.cpp b/Engine/Plugins/Experimental/GameFeatures/Source/GameFeatures/Private/GameFeaturesSubsystem.cpp
--- a/Engine/Plugins/Experimental/GameFeatures/Source/GameFeatures/Private/GameFeaturesSubsystem.cpp
+++ b/Engine/Plugins/Experimental/GameFeatures/Source/GameFeatures/Private/GameFeaturesSubsystem.cpp
@@ -827,11 +827,18 @@ void UGameFeaturesSubsystem::OnGameFeatureUnregistering(const UGameFeatureData*
 {
 	CallbackObservers(EObserverCallback::Unregistering, PluginIdentifier, &PluginName, GameFeatureData);
 
-	for (UGameFeatureAction* Action : GameFeatureData->GetActions())
+#if !WITH_EDITOR
+	check(GameFeatureData);
+#else
+	if (GameFeatureData) // In the editor the GameFeatureData asset can be force deleted, otherwise it should exist
+#endif
 	{
-		if (Action != nullptr)
+		for (UGameFeatureAction* Action : GameFeatureData->GetActions())
 		{
-			Action->OnGameFeatureUnregistering();
+			if (Action != nullptr)
+			{
+				Action->OnGameFeatureUnregistering();
+			}
 		}
 	}
 }
@@ -852,11 +859,18 @@ void UGameFeaturesSubsystem::OnGameFeatureUnloading(const UGameFeatureData* Game
 {
 	CallbackObservers(EObserverCallback::Unloading, PluginIdentifier, nullptr, GameFeatureData);
 
-	for (UGameFeatureAction* Action : GameFeatureData->GetActions())
+#if !WITH_EDITOR
+	check(GameFeatureData);
+#else
+	if (GameFeatureData) // In the editor the GameFeatureData asset can be force deleted, otherwise it should exist
+#endif
 	{
-		if (Action != nullptr)
+		for (UGameFeatureAction* Action : GameFeatureData->GetActions())
 		{
-			Action->OnGameFeatureUnloading();
+			if (Action != nullptr)
+			{
+				Action->OnGameFeatureUnloading();
+			}
 		}
 	}
 }
@@ -884,11 +898,18 @@ void UGameFeaturesSubsystem::OnGameFeatureDeactivating(const UGameFeatureData* G
 {
 	CallbackObservers(EObserverCallback::Deactivating, PluginIdentifier, &PluginName, GameFeatureData, &Context);
 
-	for (UGameFeatureAction* Action : GameFeatureData->GetActions())
+#if !WITH_EDITOR
+	check(GameFeatureData);
+#else
+	if (GameFeatureData) // In the editor the GameFeatureData asset can be force deleted, otherwise it should exist
+#endif
 	{
-		if (Action != nullptr)
+		for (UGameFeatureAction* Action : GameFeatureData->GetActions())
 		{
-			Action->OnGameFeatureDeactivating(Context);
+			if (Action != nullptr)
+			{
+				Action->OnGameFeatureDeactivating(Context);
+			}
 		}
 	}
 }
@@ -2760,7 +2781,10 @@ void UGameFeaturesSubsystem::CallbackObservers(EObserverCallback CallbackType, c
 	case EObserverCallback::Unregistering:
 	{
 		check(PluginName);
+#if !WITH_EDITOR
+		// In the editor the GameFeatureData asset can be force deleted, otherwise it should exist
 		check(GameFeatureData);
+#endif
 		for (UObject* Observer : LocalObservers)
 		{
 			CastChecked<IGameFeatureStateChangeObserver>(Observer)->OnGameFeatureUnregistering(GameFeatureData, *PluginName, PluginIdentifier.GetFullPluginURL());
@@ -2778,7 +2802,10 @@ void UGameFeaturesSubsystem::CallbackObservers(EObserverCallback CallbackType, c
 	}
 	case EObserverCallback::Unloading:
 	{
+#if !WITH_EDITOR
+		// In the editor the GameFeatureData asset can be force deleted, otherwise it should exist
 		check(GameFeatureData);
+#endif
 		for (UObject* Observer : LocalObservers)
 		{
 			CastChecked<IGameFeatureStateChangeObserver>(Observer)->OnGameFeatureUnloading(GameFeatureData, PluginIdentifier.GetFullPluginURL());
@@ -2796,7 +2823,10 @@ void UGameFeaturesSubsystem::CallbackObservers(EObserverCallback CallbackType, c
 	}
 	case EObserverCallback::Deactivating:
 	{
+#if !WITH_EDITOR
+		// In the editor the GameFeatureData asset can be force deleted, otherwise it should exist
 		check(GameFeatureData);
+#endif
 		check(StateChangeContext);
 		FGameFeatureDeactivatingContext* DeactivatingContext = static_cast<FGameFeatureDeactivatingContext*>(StateChangeContext);
 		if (ensureAlwaysMsgf(DeactivatingContext, TEXT("Invalid StateChangeContext supplied! Could not cast to FGameFeaturePauseStateChangeContext*!")))
