Add GameFeaturePlugin.ForceSyncAssetRegistryAppend CVar. This temporarily works around a race condition between the Asset Registry and GC which can lead to a fatal error.

diff --git a/Engine/Plugins/Experimental/GameFeatures/Source/GameFeatures/Private/GameFeaturePluginStateMachine.cpp b/Engine/Plugins/Experimental/GameFeatures/Source/GameFeatures/Private/GameFeaturePluginStateMachine.cpp
--- a/Engine/Plugins/Experimental/GameFeatures/Source/GameFeatures/Private/GameFeaturePluginStateMachine.cpp
+++ b/Engine/Plugins/Experimental/GameFeatures/Source/GameFeatures/Private/GameFeaturePluginStateMachine.cpp
@@ -76,6 +76,10 @@ namespace UE::GameFeatures
 		true,
 		TEXT("Enable to force shaderlibs to be opened on the game thread"));
 
+	static TAutoConsoleVariable<bool> CVarForceSyncAssetRegistryAppend(TEXT("GameFeaturePlugin.ForceSyncAssetRegistryAppend"),
+		true,
+		TEXT("Enable to force calls to IAssetRegistry::AppendState to happen on the game thread"));
+
 	#define GAME_FEATURE_PLUGIN_STATE_TO_STRING(inEnum, inText) case EGameFeaturePluginState::inEnum: return TEXT(#inEnum);
 	FString ToString(EGameFeaturePluginState InType)
 	{
@@ -2104,18 +2108,22 @@ struct FGameFeaturePluginState_Mounting : public FGameFeaturePluginState
 			return;
 		}
 
-		UE::Tasks::Launch(UE_SOURCE_LOCATION, [this, PluginAssetRegistry=MoveTemp(PluginAssetRegistry)]
+		const bool bForceSyncAssetRegistryAppend = UE::GameFeatures::CVarForceSyncAssetRegistryAppend.GetValueOnGameThread();
+		UE::Tasks::Launch(UE_SOURCE_LOCATION, [this, PluginAssetRegistry=MoveTemp(PluginAssetRegistry), bForceSyncAssetRegistryAppend]
 		{
 			bool bSuccess = false;
 			TSharedPtr<FAssetRegistryState> PluginAssetRegistryState = MakeShared<FAssetRegistryState>();
 			if (FAssetRegistryState::LoadFromDisk(*PluginAssetRegistry, FAssetRegistryLoadOptions(), *PluginAssetRegistryState))
 			{
 				IAssetRegistry& AssetRegistry = UAssetManager::Get().GetAssetRegistry();
-				AssetRegistry.AppendState(*PluginAssetRegistryState);
+				if (!bForceSyncAssetRegistryAppend)
+				{
+					AssetRegistry.AppendState(*PluginAssetRegistryState);
+				}
 				bSuccess = true;
 			}
 
-			ExecuteOnGameThread(UE_SOURCE_LOCATION, [this, bSuccess]
+			ExecuteOnGameThread(UE_SOURCE_LOCATION, [this, PluginAssetRegistryState, bSuccess, bForceSyncAssetRegistryAppend]
 			{
 				TRACE_CPUPROFILER_EVENT_SCOPE(GFP_Mounting_ARComplete);
 
@@ -2123,6 +2131,11 @@ struct FGameFeaturePluginState_Mounting : public FGameFeaturePluginState
 				{
 					Result = GetErrorResult(TEXT("Failed_To_Load_Plugin_AssetRegistry"));
 				}
+				else if (bForceSyncAssetRegistryAppend)
+				{
+					IAssetRegistry& AssetRegistry = UAssetManager::Get().GetAssetRegistry();
+					AssetRegistry.AppendState(*PluginAssetRegistryState);
+				}
 
 				CompletedSubStates |= ESubState::LoadAssetRegistry;
 				UpdateStateMachineImmediate();
