Add profiling markers for game feature actions unloading.

#rnx

#tests Tested in PIE with Unreal Insights. Made sure I noticed the markers in a profiling session.

diff --git a/Engine/Plugins/Experimental/GameFeatures/Source/GameFeatures/Private/GameFeaturesSubsystem.cpp b/Engine/Plugins/Experimental/GameFeatures/Source/GameFeatures/Private/GameFeaturesSubsystem.cpp
--- a/Engine/Plugins/Experimental/GameFeatures/Source/GameFeatures/Private/GameFeaturesSubsystem.cpp
+++ b/Engine/Plugins/Experimental/GameFeatures/Source/GameFeatures/Private/GameFeaturesSubsystem.cpp
@@ -896,14 +896,17 @@ void UGameFeaturesSubsystem::OnGameFeatureActivating(const UGameFeatureData* Gam
 
 void UGameFeaturesSubsystem::OnGameFeatureDeactivating(const UGameFeatureData* GameFeatureData, const FString& PluginName, FGameFeatureDeactivatingContext& Context, const FGameFeaturePluginIdentifier& PluginIdentifier)
 {
-	CallbackObservers(EObserverCallback::Deactivating, PluginIdentifier, &PluginName, GameFeatureData, &Context);
-
 #if !WITH_EDITOR
 	check(GameFeatureData);
 #else
 	if (GameFeatureData) // In the editor the GameFeatureData asset can be force deleted, otherwise it should exist
 #endif
 	{
+		TRACE_CPUPROFILER_EVENT_SCOPE(GFP_OnDeactivating_CallbackObservers);
+		CallbackObservers(EObserverCallback::Deactivating, PluginIdentifier, &PluginName, GameFeatureData, &Context);
+	}
+	{
+		TRACE_CPUPROFILER_EVENT_SCOPE(GFP_OnDeactivating_OnGameFeatureDeactivating);
 		for (UGameFeatureAction* Action : GameFeatureData->GetActions())
 		{
 			if (Action != nullptr)
