Add some more protection when additional package managers exist to not have assets pulled into a non game based pak chunk.
When multiple package managers exist but we don't have multiple chunkIds it's likey there's a configuration issue but we can add protection around this and help catch something is off.

diff --git a/Engine/Plugins/Experimental/GameFeatures/Source/GameFeatures/Private/GameFeatureAction_AddChunkOverride.cpp b/Engine/Plugins/Experimental/GameFeatures/Source/GameFeatures/Private/GameFeatureAction_AddChunkOverride.cpp
--- a/Engine/Plugins/Experimental/GameFeatures/Source/GameFeatures/Private/GameFeatureAction_AddChunkOverride.cpp
+++ b/Engine/Plugins/Experimental/GameFeatures/Source/GameFeatures/Private/GameFeatureAction_AddChunkOverride.cpp
@@ -38,14 +38,21 @@ void UGameFeatureAction_AddChunkOverride::OnGameFeatureUnregistering()
 }
 
 #if WITH_EDITOR
-void UGameFeatureAction_AddChunkOverride::GetChunkForPackage(const FString& PackageName, const int32 DefaultGameChunk, TArray<int32>& OutChunkList)
+void UGameFeatureAction_AddChunkOverride::GetChunkForPackage(const FName PackageName, const int32 DefaultGameChunk, TArray<int32>& OutChunkList)
+{
+	TSet<FPrimaryAssetId> Managers;
+	UAssetManager::Get().GetPackageManagers(PackageName, true, Managers);
+	GetChunkForPackage(PackageName.ToString(), Managers, DefaultGameChunk, OutChunkList);
+}
+
+void UGameFeatureAction_AddChunkOverride::GetChunkForPackage(const FString& PackageName, const TSet<FPrimaryAssetId>& Managers, const int32 DefaultGameChunk, TArray<int32>& OutChunkList)
 {
 	if (GameFeatureAction_AddChunkOverride::PluginToChunkId.Num() == 0)
 	{
 		return;
 	}
 
-	auto ResolveMultiChunkDependecies = [PackageName, DefaultGameChunk , &OutChunkList]()
+	auto ResolveMultiChunkDependecies = [PackageName, &Managers, DefaultGameChunk,  &OutChunkList]()
 	{
 		UE_LOG(LogAddChunkOverride, VeryVerbose, TEXT("%s was referenced by one than one chunk"), *PackageName);
 		for (const int32 OutChunkId : OutChunkList)
@@ -59,6 +66,26 @@ void UGameFeatureAction_AddChunkOverride::GetChunkForPackage(const FString& Pack
 			OutChunkList.Reset();
 			OutChunkList.Add(DefaultGameChunk);
 		}
+		else
+		{
+			// If multiple package mangers exist for this package with different chunk IDs it should go into the default game chunk.
+			UAssetManager& AssetManager = UAssetManager::Get();
+			TSet<int32> ManagerChunkIds;
+			for (const FPrimaryAssetId& PrimaryAssetId : Managers)
+			{
+				FPrimaryAssetRules Rules = AssetManager.GetPrimaryAssetRules(PrimaryAssetId);
+				if (Rules.ChunkId != INDEX_NONE)
+				{
+					ManagerChunkIds.Add(Rules.ChunkId);
+				}
+			}
+			if (ManagerChunkIds.Num() > 1)
+			{
+				UE_LOG(LogAddChunkOverride, Warning, TEXT("Forcing %s into gameplay chunk %d. It was referend by multiple GFPs which might load at different times. Package managers with a valid chunkID might not have been registered for this type."), *PackageName, DefaultGameChunk);
+				OutChunkList.Reset();
+				OutChunkList.Add(DefaultGameChunk);
+			}
+		}
 	};
 
 	static const FString EngineDir(TEXT("/Engine/"));
diff --git a/Engine/Plugins/Experimental/GameFeatures/Source/GameFeatures/Public/GameFeatureAction_AddChunkOverride.h b/Engine/Plugins/Experimental/GameFeatures/Source/GameFeatures/Public/GameFeatureAction_AddChunkOverride.h
--- a/Engine/Plugins/Experimental/GameFeatures/Source/GameFeatures/Public/GameFeatureAction_AddChunkOverride.h
+++ b/Engine/Plugins/Experimental/GameFeatures/Source/GameFeatures/Public/GameFeatureAction_AddChunkOverride.h
@@ -43,7 +43,8 @@ public:
 	 * 
 	 * This can be necessary to reassign startup packages such as the GameFeatureData asset.
 	 */
-	GAMEFEATURES_API static void GetChunkForPackage(const FString& PackageName, const int32 DefaultGameChunk, TArray<int32>& OutChunkList);
+	GAMEFEATURES_API static void GetChunkForPackage(const FName PackageName, const int32 DefaultGameChunk, TArray<int32>& OutChunkList);
+	GAMEFEATURES_API static void GetChunkForPackage(const FString& PackageName, const TSet<FPrimaryAssetId>& Managers, const int32 DefaultGameChunk, TArray<int32>& OutChunkList);
 
 	GAMEFEATURES_API static FString GetPluginNameFromChunkID(int32 ChunkID);
 
