Downgrading warning

diff --git a/Engine/Plugins/Experimental/GameFeatures/Source/GameFeatures/Private/GameFeatureAction_AddChunkOverride.cpp b/Engine/Plugins/Experimental/GameFeatures/Source/GameFeatures/Private/GameFeatureAction_AddChunkOverride.cpp
--- a/Engine/Plugins/Experimental/GameFeatures/Source/GameFeatures/Private/GameFeatureAction_AddChunkOverride.cpp
+++ b/Engine/Plugins/Experimental/GameFeatures/Source/GameFeatures/Private/GameFeatureAction_AddChunkOverride.cpp
@@ -81,7 +81,7 @@ void UGameFeatureAction_AddChunkOverride::GetChunkForPackage(const FString& Pack
 			}
 			if (ManagerChunkIds.Num() > 1)
 			{
-				UE_LOG(LogAddChunkOverride, Warning, TEXT("Forcing %s into gameplay chunk %d. It was referend by multiple GFPs which might load at different times. Package managers with a valid chunkID might not have been registered for this type."), *PackageName, DefaultGameChunk);
+				UE_LOG(LogAddChunkOverride, Log, TEXT("Forcing %s into gameplay chunk %d. It was referend by multiple GFPs which might load at different times. Package managers with a valid chunkID might not have been registered for this type."), *PackageName, DefaultGameChunk);
 				OutChunkList.Reset();
 				OutChunkList.Add(DefaultGameChunk);
 			}
