Shorten GFP plugin identifier to just be the leaf file name.  This avoids issues with non-normalized paths.

diff --git a/Engine/Plugins/Experimental/GameFeatures/Source/GameFeatures/Private/GameFeaturePluginStateMachine.cpp b/Engine/Plugins/Experimental/GameFeatures/Source/GameFeatures/Private/GameFeaturePluginStateMachine.cpp
--- a/Engine/Plugins/Experimental/GameFeatures/Source/GameFeatures/Private/GameFeaturePluginStateMachine.cpp
+++ b/Engine/Plugins/Experimental/GameFeatures/Source/GameFeatures/Private/GameFeaturePluginStateMachine.cpp
@@ -3610,8 +3610,13 @@ bool FGameFeaturePluginStateMachineProperties::ParseURL()
 		return false;
 	}
 
-	//The current IdentifyingURLSubset IS the same as the PluginInstalledFilename, so just use the already parsed version
-	PluginInstalledFilename = PluginIdentifier.IdentifyingURLSubset;
+	FStringView PluginPathFromURL;
+	if (!UGameFeaturesSubsystem::ParsePluginURL(PluginIdentifier.GetFullPluginURL(), nullptr, &PluginPathFromURL))
+	{
+		return false;
+	}
+
+	PluginInstalledFilename = PluginPathFromURL;
 	PluginName = FPaths::GetBaseFilename(PluginInstalledFilename);
 
 	if (PluginInstalledFilename.IsEmpty() || !PluginInstalledFilename.EndsWith(TEXT(".uplugin")))
diff --git a/Engine/Plugins/Experimental/GameFeatures/Source/GameFeatures/Private/GameFeaturesSubsystem.cpp b/Engine/Plugins/Experimental/GameFeatures/Source/GameFeatures/Private/GameFeaturesSubsystem.cpp
--- a/Engine/Plugins/Experimental/GameFeatures/Source/GameFeatures/Private/GameFeaturesSubsystem.cpp
+++ b/Engine/Plugins/Experimental/GameFeatures/Source/GameFeatures/Private/GameFeaturesSubsystem.cpp
@@ -199,26 +199,12 @@ void FGameFeaturePluginIdentifier::FromPluginURL(FString PluginURLIn)
 	IdentifyingURLSubset.Reset();
 	PluginURL = MoveTemp(PluginURLIn);
 
-	PluginProtocol = UGameFeaturesSubsystem::GetPluginURLProtocol(PluginURL);
-
-	if (ensureAlwaysMsgf( ((PluginProtocol != EGameFeaturePluginProtocol::Unknown) 
-						&& (PluginProtocol != EGameFeaturePluginProtocol::Count)),
-						TEXT("Invalid PluginProtocol in PluginURL %s"), *PluginURL))
+	if (UGameFeaturesSubsystem::ParsePluginURL(PluginURL, &PluginProtocol, &IdentifyingURLSubset))
 	{
-
-		int32 PluginProtocolEndIndex = FCString::Strlen(UE::GameFeatures::GameFeaturePluginProtocolPrefix(PluginProtocol));
-		int32 FirstOptionIndex = PluginURL.Find(UE::GameFeatures::PluginURLStructureInfo::OptionSeperator, ESearchCase::IgnoreCase, ESearchDir::FromStart, PluginProtocolEndIndex);
-		
-		//If we don't have any options, then the IdentifyingURLSubset is just our entire URL except the protocol string
-		if (FirstOptionIndex == INDEX_NONE)
-		{
-			IdentifyingURLSubset = FStringView(PluginURL).RightChop(PluginProtocolEndIndex);
-		}
-		//The IdentifyingURLSubset will be the string between the end of the protocol string and before the first option
-		else
+		if (!IdentifyingURLSubset.IsEmpty())
 		{
-			const int32 IdentifierCharCount = (FirstOptionIndex - PluginProtocolEndIndex);
-			IdentifyingURLSubset = FStringView(PluginURL).Mid(PluginProtocolEndIndex, IdentifierCharCount);
+			// Plugins must be unique so just use the name as the identifier. This avoids issues with normalizing paths.
+			IdentifyingURLSubset = FPathViews::GetCleanFilename(IdentifyingURLSubset);
 		}
 	}
 }
@@ -775,6 +761,55 @@ bool UGameFeaturesSubsystem::IsPluginURLProtocol(FStringView PluginURL, EGameFea
 	return PluginURL.StartsWith(UE::GameFeatures::GameFeaturePluginProtocolPrefix(PluginProtocol));
 }
 
+bool UGameFeaturesSubsystem::ParsePluginURL(FStringView PluginURL, EGameFeaturePluginProtocol* OutProtocol /*= nullptr*/, FStringView* OutPath /*= nullptr*/, FStringView* OutOptions /*= nullptr*/)
+{
+	FStringView Path;
+	FStringView Options;
+	EGameFeaturePluginProtocol PluginProtocol = UGameFeaturesSubsystem::GetPluginURLProtocol(PluginURL);
+
+	if (ensureAlwaysMsgf(PluginProtocol != EGameFeaturePluginProtocol::Unknown && PluginProtocol != EGameFeaturePluginProtocol::Count,
+		TEXT("Invalid PluginProtocol in PluginURL %.*s"), PluginURL.Len(), PluginURL.GetData()))
+	{
+		int32 PluginProtocolLen = FCString::Strlen(UE::GameFeatures::GameFeaturePluginProtocolPrefix(PluginProtocol));
+		int32 FirstOptionIndex = UE::String::FindFirst(PluginURL, UE::GameFeatures::PluginURLStructureInfo::OptionSeperator, ESearchCase::IgnoreCase);
+
+		//If we don't have any options, then the Path is just our entire URL except the protocol string
+		if (FirstOptionIndex == INDEX_NONE)
+		{
+			Path = PluginURL.RightChop(PluginProtocolLen);
+		}
+		//The Path will be the string between the end of the protocol string and before the first option
+		else
+		{
+			const int32 IdentifierCharCount = (FirstOptionIndex - PluginProtocolLen);
+			Path = PluginURL.Mid(PluginProtocolLen, IdentifierCharCount);
+			Options = PluginURL.RightChop(FirstOptionIndex);
+		}
+
+		if (ensureAlwaysMsgf(Path.EndsWith(TEXTVIEW(".uplugin")), TEXT("Invalid path in PluginURL %.*s"), PluginURL.Len(), PluginURL.GetData()))
+		{
+			if (OutProtocol)
+			{
+				*OutProtocol = PluginProtocol;
+			}
+
+			if (OutPath)
+			{
+				*OutPath = Path;
+			}
+
+			if (OutOptions)
+			{
+				*OutOptions = Options;
+			}
+
+			return true;
+		}
+	}
+
+	return false;
+}
+
 void UGameFeaturesSubsystem::OnGameFeatureTerminating(const FString& PluginName, const FGameFeaturePluginIdentifier& PluginIdentifier)
 {
 	CallbackObservers(EObserverCallback::Terminating, PluginIdentifier, &PluginName);
@@ -1863,8 +1898,13 @@ bool UGameFeaturesSubsystem::GetBuiltInGameFeaturePluginDetails(const TSharedRef
 
 bool UGameFeaturesSubsystem::GetGameFeaturePluginDetails(FString PluginURL, FGameFeaturePluginDetails& OutPluginDetails) const
 {
-	FGameFeaturePluginIdentifier Ident(MoveTemp(PluginURL));
-	return GetGameFeaturePluginDetailsInternal(FString(Ident.GetIdentifyingString()), OutPluginDetails);
+	FStringView PluginPath;
+	if (UGameFeaturesSubsystem::ParsePluginURL(PluginURL, nullptr, &PluginPath))
+	{
+		return GetGameFeaturePluginDetailsInternal(FString(PluginPath), OutPluginDetails);
+	}
+
+	return false;
 }
 
 bool UGameFeaturesSubsystem::GetGameFeaturePluginDetailsInternal(const FString& PluginDescriptorFilename, FGameFeaturePluginDetails& OutPluginDetails) const
diff --git a/Engine/Plugins/Experimental/GameFeatures/Source/GameFeatures/Public/GameFeaturesSubsystem.h b/Engine/Plugins/Experimental/GameFeatures/Source/GameFeatures/Public/GameFeaturesSubsystem.h
--- a/Engine/Plugins/Experimental/GameFeatures/Source/GameFeatures/Public/GameFeaturesSubsystem.h
+++ b/Engine/Plugins/Experimental/GameFeatures/Source/GameFeatures/Public/GameFeaturesSubsystem.h
@@ -474,6 +474,9 @@ public:
 	/** Tests whether the plugin URL is the specified protocol */
 	static bool IsPluginURLProtocol(FStringView PluginURL, EGameFeaturePluginProtocol PluginProtocol);
 
+	/** Parse the plugin URL into subparts */
+	static bool ParsePluginURL(FStringView PluginURL, EGameFeaturePluginProtocol* OutProtocol = nullptr, FStringView* OutPath = nullptr, FStringView* OutOptions = nullptr);
+
 public:
 	/** Returns all the active plugins GameFeatureDatas */
 	void GetGameFeatureDataForActivePlugins(TArray<const UGameFeatureData*>& OutActivePluginFeatureDatas);
