Fix CoD dialog sometimes getting queued when there is nothing to download.

diff --git a/Engine/Plugins/Experimental/GameFeatures/Source/GameFeatures/Private/GameFeaturesSubsystem.cpp b/Engine/Plugins/Experimental/GameFeatures/Source/GameFeatures/Private/GameFeaturesSubsystem.cpp
--- a/Engine/Plugins/Experimental/GameFeatures/Source/GameFeatures/Private/GameFeaturesSubsystem.cpp
+++ b/Engine/Plugins/Experimental/GameFeatures/Source/GameFeatures/Private/GameFeaturesSubsystem.cpp
@@ -2198,6 +2198,14 @@ struct FGameFeaturePluginPredownloadContext : public FGameFeaturePluginPredownlo
 			return;
 		}
 
+		// Early out if everything is up to date already. This helps avoid enquing UI dialogs for content that doesn't actually need to be downloaded
+		TValueOrError<FInstallBundleCombinedInstallState, EInstallBundleResult> MaybeInstallState = BundleManager->GetInstallStateSynchronous(BundlesToInstall, false);
+		if (MaybeInstallState.HasValue() && MaybeInstallState.GetValue().GetAllBundlesHaveState(EInstallBundleInstallState::UpToDate))
+		{
+			SetComplete(MakeValue());
+			return;
+		}
+
 		BundleManager->GetContentState(BundlesToInstall, EInstallBundleGetContentStateFlags::None, false,
 			FInstallBundleGetContentStateDelegate::CreateLambda([Context = SharedThis(this)](FInstallBundleCombinedContentState BundleContentState)
 			{ Context->OnGotContentState(MoveTemp(BundleContentState)); }));
