Add GameFeatureOptionalContentInstaller
- Can be configured to deliver optional cached content whenver certain GFPs are loaded.
- Disabled by default for now.

diff --git a/Engine/Plugins/Experimental/GameFeatures/Source/GameFeatures/Private/GameFeatureOptionalContentInstaller.cpp b/Engine/Plugins/Experimental/GameFeatures/Source/GameFeatures/Private/GameFeatureOptionalContentInstaller.cpp
new file mode 100644
--- /dev/null
+++ b/Engine/Plugins/Experimental/GameFeatures/Source/GameFeatures/Private/GameFeatureOptionalContentInstaller.cpp
@@ -0,0 +1,317 @@
+// Copyright Epic Games, Inc. All Rights Reserved.
+
+#include "GameFeatureOptionalContentInstaller.h"
+
+#include "Algo/AllOf.h"
+#include "GameFeaturePluginOperationResult.h"
+#include "GameFeaturesSubsystem.h"
+#include "GameFeatureTypes.h"
+#include "Logging/StructuredLog.h"
+
+#include UE_INLINE_GENERATED_CPP_BY_NAME(GameFeatureOptionalContentInstaller)
+
+namespace GameFeatureOptionalContentInstaller
+{
+	static const ELogVerbosity::Type InstallBundleManagerVerbosityOverride = ELogVerbosity::Verbose;
+
+	static const FStringView ErrorNamespace = TEXTVIEW("GameFeaturePlugin.OptionalDownload.");
+
+	static TAutoConsoleVariable<bool> CVarEnableOptionalContentInstaller(TEXT("GameFeatureOptionalContentInstaller.Enable"), 
+		false,
+		TEXT("Enable optional content installer"));
+}
+
+TMulticastDelegate<void(const FString& PluginName, const UE::GameFeatures::FResult&)> UGameFeatureOptionalContentInstaller::OnOptionalContentInstalled;
+
+void UGameFeatureOptionalContentInstaller::BeginDestroy()
+{
+	IConsoleManager::Get().UnregisterConsoleVariableSink_Handle(CVarSinkHandle);
+	Super::BeginDestroy();
+}
+
+void UGameFeatureOptionalContentInstaller::Init(TUniqueFunction<TArray<FName>(FString)> InGetOptionalBundlePredicate)
+{
+	GetOptionalBundlePredicate = MoveTemp(InGetOptionalBundlePredicate);
+	BundleManager = IInstallBundleManager::GetPlatformInstallBundleManager();
+
+	// Create the cvar sink and set up the thread pool
+	CVarSinkHandle = IConsoleManager::Get().RegisterConsoleVariableSink_Handle(
+		FConsoleCommandDelegate::CreateUObject(this, &UGameFeatureOptionalContentInstaller::OnCVarsChanged));
+	bEnabledCVar = GameFeatureOptionalContentInstaller::CVarEnableOptionalContentInstaller.GetValueOnGameThread();
+}
+
+void UGameFeatureOptionalContentInstaller::Enable(bool bInEnable)
+{
+	bool bOldEnabled = IsEnabled();
+	bEnabled = bInEnable;
+	bEnabledCVar = GameFeatureOptionalContentInstaller::CVarEnableOptionalContentInstaller.GetValueOnGameThread();
+	bool bNewEnabled = IsEnabled();
+
+	if (bOldEnabled != bNewEnabled)
+	{
+		if (bNewEnabled)
+		{
+			OnEnabled();
+		}
+		else
+		{
+			OnDisabled();
+		}
+	}
+}
+
+void UGameFeatureOptionalContentInstaller::EnableCellularDownloading(bool bEnable)
+{
+	if (bAllowCellDownload == bEnable)
+	{
+		return;
+	}
+
+	bAllowCellDownload = bEnable;
+
+	// Update flags on active requests
+	for ( TPair<FString, FGFPInstall>& Pair : ActiveGFPInstalls)
+	{
+		BundleManager->UpdateContentRequestFlags(Pair.Value.BundlesEnqueued,
+			bEnable ? EInstallBundleRequestFlags::None : EInstallBundleRequestFlags::CheckForCellularDataUsage,
+			bEnable ? EInstallBundleRequestFlags::CheckForCellularDataUsage : EInstallBundleRequestFlags::None);
+	}
+}
+
+bool UGameFeatureOptionalContentInstaller::UpdateContent(const FString& PluginName, bool bIsPredownload)
+{
+	TArray<FName> Bundles = GetOptionalBundlePredicate(PluginName);
+
+	bool bIsAvailable = false;
+	if (!Bundles.IsEmpty())
+	{
+		TValueOrError<FInstallBundleCombinedInstallState, EInstallBundleResult> MaybeInstallState = BundleManager->GetInstallStateSynchronous(Bundles, false);
+		if (MaybeInstallState.HasValue())
+		{
+			const FInstallBundleCombinedInstallState& InstallState = MaybeInstallState.GetValue();
+			bIsAvailable = Algo::AllOf(Bundles, [&InstallState](FName BundleName) { return InstallState.IndividualBundleStates.Contains(BundleName); });
+		}
+	}
+
+	if (!bIsAvailable)
+	{
+		return false;
+	}
+
+	EInstallBundleRequestFlags InstallFlags = EInstallBundleRequestFlags::AsyncMount;
+	if (bIsPredownload)
+	{
+		InstallFlags |= EInstallBundleRequestFlags::SkipMount;
+	}
+	if (!bAllowCellDownload)
+	{
+		InstallFlags |= EInstallBundleRequestFlags::CheckForCellularDataUsage;
+	}
+
+	TValueOrError<FInstallBundleRequestInfo, EInstallBundleResult> MaybeRequest = BundleManager->RequestUpdateContent(
+		Bundles, 
+		InstallFlags, 
+		GameFeatureOptionalContentInstaller::InstallBundleManagerVerbosityOverride);
+
+	if (MaybeRequest.HasError())
+	{
+		UE_LOGFMT(LogGameFeatures, Error, "Failed to request optional content for GFP {GFP}, Error: {Error}", 
+			("GFP", PluginName),
+			("Error", LexToString(MaybeRequest.GetError())));
+
+		UE::GameFeatures::FResult ErrorResult = MakeError(FString::Printf(TEXT("%.*s%s"),
+			GameFeatureOptionalContentInstaller::ErrorNamespace.Len(), GameFeatureOptionalContentInstaller::ErrorNamespace.GetData(),
+			LexToString(MaybeRequest.GetError())));
+		OnOptionalContentInstalled.Broadcast(PluginName, ErrorResult);
+
+		return false;
+	}
+
+	FInstallBundleRequestInfo& Request =  MaybeRequest.GetValue();
+	if (!Request.BundlesEnqueued.IsEmpty())
+	{
+		FGFPInstall& Pending = ActiveGFPInstalls.FindOrAdd(PluginName);
+
+		if (!Pending.CallbackHandle.IsValid())
+		{
+			Pending.CallbackHandle = IInstallBundleManager::InstallBundleCompleteDelegate.AddUObject(this,
+				&UGameFeatureOptionalContentInstaller::OnContentInstalled, PluginName);
+		}
+
+		// This should overwrite any previous pending request info
+		Pending.BundlesEnqueued = MoveTemp(Request.BundlesEnqueued);
+		Pending.bIsPredownload = bIsPredownload;
+	}
+
+	return true;
+}
+
+void UGameFeatureOptionalContentInstaller::OnContentInstalled(FInstallBundleRequestResultInfo InResult, FString PluginName)
+{
+	FGFPInstall* MaybeInstall = ActiveGFPInstalls.Find(PluginName);
+	if (!MaybeInstall)
+	{
+		return;
+	}
+
+	FGFPInstall& GFPInstall = *MaybeInstall;
+	if (!GFPInstall.BundlesEnqueued.Contains(InResult.BundleName))
+	{
+		return;
+	}
+
+	GFPInstall.BundlesEnqueued.Remove(InResult.BundleName);
+
+	if (InResult.Result != EInstallBundleResult::OK)
+	{
+		if (InResult.OptionalErrorCode.IsEmpty())
+		{
+			UE_LOGFMT(LogGameFeatures, Error, "Failed to install optional bundle {Bundle} for GFP {GFP}, Error: {Error}",
+				("Bundle", InResult.BundleName),
+				("GFP", PluginName),
+				("Error", LexToString(InResult.Result)));
+		}
+		else
+		{
+			UE_LOGFMT(LogGameFeatures, Error, "Failed to install optional bundle {Bundle} for GFP {GFP}, Error: {Error}",
+				("Bundle", InResult.BundleName),
+				("GFP", PluginName),
+				("Error", InResult.OptionalErrorCode));
+		}
+
+		//Use OptionalErrorCode and/or OptionalErrorText if available
+		const FString ErrorCodeEnding = (InResult.OptionalErrorCode.IsEmpty()) ? LexToString(InResult.Result) : InResult.OptionalErrorCode;
+		FText ErrorText = InResult.OptionalErrorCode.IsEmpty() ? UE::GameFeatures::CommonErrorCodes::GetErrorTextForBundleResult(InResult.Result) : InResult.OptionalErrorText;
+		UE::GameFeatures::FResult ErrorResult = UE::GameFeatures::FResult(
+			MakeError(FString::Printf(TEXT("%.*s%s"), GameFeatureOptionalContentInstaller::ErrorNamespace.Len(), GameFeatureOptionalContentInstaller::ErrorNamespace.GetData(), *ErrorCodeEnding)),
+			MoveTemp(ErrorText)
+		);
+		OnOptionalContentInstalled.Broadcast(PluginName, ErrorResult);
+
+		// Cancel any remaining downloads
+		BundleManager->CancelUpdateContent(GFPInstall.BundlesEnqueued);
+	}
+
+	if (GFPInstall.BundlesEnqueued.IsEmpty())
+	{
+		if (GFPInstall.bIsPredownload)
+		{
+			// Predownload shouldn't pin any cached bundles so release them now
+
+			// Delay call to ReleaseBundlesIfPossible. We don't want to release them from within the complete callback.
+			FTSTicker::GetCoreTicker().AddTicker(FTickerDelegate::CreateWeakLambda(this, 
+			[this, PluginName, bInstalled = InResult.bContentWasInstalled](float)
+			{
+				// A machine is active, don't release
+				if (!RelevantGFPs.Contains(PluginName))
+				{
+					ReleaseContent(PluginName);
+				}
+
+				if (bInstalled)
+				{
+					OnOptionalContentInstalled.Broadcast(PluginName, MakeValue());
+				}
+
+				return false;
+			}));
+		}
+		else if (InResult.bContentWasInstalled)
+		{
+			OnOptionalContentInstalled.Broadcast(PluginName, MakeValue());
+		}
+
+		// book keeping
+		IInstallBundleManager::InstallBundleCompleteDelegate.Remove(GFPInstall.CallbackHandle);
+		ActiveGFPInstalls.Remove(PluginName);
+	}
+}
+
+void UGameFeatureOptionalContentInstaller::ReleaseContent(const FString& PluginName)
+{
+	TArray<FName> Bundles = GetOptionalBundlePredicate(PluginName);
+	if (Bundles.IsEmpty())
+	{
+		return;
+	}
+
+	BundleManager->RequestReleaseContent(
+		Bundles, 
+		EInstallBundleReleaseRequestFlags::None, 
+		{}, 
+		GameFeatureOptionalContentInstaller::InstallBundleManagerVerbosityOverride);
+}
+
+void UGameFeatureOptionalContentInstaller::OnEnabled()
+{
+	ensure(RelevantGFPs.IsEmpty());
+	RelevantGFPs.Empty();
+
+	UGameFeaturesSubsystem::Get().ForEachGameFeature([this](FGameFeatureInfo&& Info) -> void
+	{
+		if (UGameFeaturesSubsystem::GetPluginURLProtocol(Info.URL) == EGameFeaturePluginProtocol::InstallBundle &&
+			Info.CurrentState >= EGameFeaturePluginState::Downloading)
+		{
+			if (UpdateContent(Info.Name, false))
+			{
+				RelevantGFPs.Add(Info.Name);
+			}
+		}
+	});
+}
+
+void UGameFeatureOptionalContentInstaller::OnDisabled()
+{
+	for (const FString& GFP : RelevantGFPs)
+	{
+		ReleaseContent(GFP);
+	}
+
+	RelevantGFPs.Empty();
+}
+
+bool UGameFeatureOptionalContentInstaller::IsEnabled() const
+{
+	return bEnabled && bEnabledCVar;
+}
+
+void UGameFeatureOptionalContentInstaller::OnCVarsChanged()
+{
+	Enable(bEnabled); // Check if CVar changed IsEnabled() and if so, call callbacks
+}
+
+void UGameFeatureOptionalContentInstaller::OnGameFeaturePredownloading(const FString& PluginName, const FGameFeaturePluginIdentifier& PluginIdentifier)
+{
+	if (!IsEnabled())
+	{
+		return;
+	}
+
+	UpdateContent(PluginName, true);
+	// Predownloads are not 'relevant', they don't have an active state machine
+}
+
+void UGameFeatureOptionalContentInstaller::OnGameFeatureDownloading(const FString& PluginName, const FGameFeaturePluginIdentifier& PluginIdentifier)
+{
+	if (!IsEnabled())
+	{
+		return;
+	}
+
+	if (UpdateContent(PluginName, false))
+	{
+		RelevantGFPs.Add(PluginName);
+	}
+}
+
+void UGameFeatureOptionalContentInstaller::OnGameFeatureReleasing(const FString& PluginName, const FGameFeaturePluginIdentifier& PluginIdentifier)
+{
+	if (!IsEnabled())
+	{
+		return;
+	}
+
+	ReleaseContent(PluginName);
+
+	RelevantGFPs.Remove(PluginName);
+}
diff --git a/Engine/Plugins/Experimental/GameFeatures/Source/GameFeatures/Private/GameFeaturePluginStateMachine.cpp b/Engine/Plugins/Experimental/GameFeatures/Source/GameFeatures/Private/GameFeaturePluginStateMachine.cpp
--- a/Engine/Plugins/Experimental/GameFeatures/Source/GameFeatures/Private/GameFeaturePluginStateMachine.cpp
+++ b/Engine/Plugins/Experimental/GameFeatures/Source/GameFeatures/Private/GameFeaturePluginStateMachine.cpp
@@ -1182,6 +1182,16 @@ struct FGameFeaturePluginState_Releasing : public FBaseDataReleaseGameFeaturePlu
 		: FBaseDataReleaseGameFeaturePluginState(InStateProperties)
 	{}
 
+	virtual void BeginState() override
+	{
+		if (ShouldReleaseContent())
+		{
+			UGameFeaturesSubsystem::Get().OnGameFeatureReleasing(StateProperties.PluginName, StateProperties.PluginIdentifier);
+		}
+
+		FBaseDataReleaseGameFeaturePluginState::BeginState();
+	}
+
 	virtual EGameFeaturePluginState GetSuccessTransitionState() const override
 	{
 		return EGameFeaturePluginState::StatusKnown;
@@ -1439,6 +1449,8 @@ struct FGameFeaturePluginState_Downloading : public FGameFeaturePluginState
 		check(StateProperties.GetPluginProtocol() == EGameFeaturePluginProtocol::InstallBundle);
 		ensureMsgf(AllowAsyncLoading(), TEXT("FGameFeaturePluginState::AllowAsyncLoading is false while attempting to download GFP data."));
 
+		UGameFeaturesSubsystem::Get().OnGameFeatureDownloading(StateProperties.PluginName, StateProperties.PluginIdentifier);
+
 		TSharedPtr<IInstallBundleManager> BundleManager = IInstallBundleManager::GetPlatformInstallBundleManager();
 		const TArray<FName>& InstallBundles = StateProperties.ProtocolMetadata.GetSubtype<FInstallBundlePluginProtocolMetaData>().InstallBundles;
 
@@ -1660,6 +1672,8 @@ struct FGameFeaturePluginState_Unmounting : public FGameFeaturePluginState
 			return;
 		}
 
+		UGameFeaturesSubsystem::Get().OnGameFeatureReleasing(StateProperties.PluginName, StateProperties.PluginIdentifier);
+
 		TSharedPtr<IInstallBundleManager> BundleManager = IInstallBundleManager::GetPlatformInstallBundleManager();
 
 		const TArray<FName>& InstallBundles = StateProperties.ProtocolMetadata.GetSubtype<FInstallBundlePluginProtocolMetaData>().InstallBundles;
@@ -2579,33 +2593,15 @@ struct FGameFeaturePluginState_Unloading : public FGameFeaturePluginState
 {
 	FGameFeaturePluginState_Unloading(FGameFeaturePluginStateMachineProperties& InStateProperties) : FGameFeaturePluginState(InStateProperties) {}
 
-	bool bHasUnloaded = false;
-
 	virtual void BeginState() override
 	{
-		bHasUnloaded = false;
 	}
 
 	virtual void UpdateState(FGameFeaturePluginStateStatus& StateStatus) override
 	{
-		if (bHasUnloaded)
-		{
-			StateStatus.SetTransition(EGameFeaturePluginState::Registered);
-			return;
-		}
-
 		UnloadGameFeatureBundles(StateProperties.GameFeatureData);
 		UGameFeaturesSubsystem::Get().OnGameFeatureUnloading(StateProperties.GameFeatureData, StateProperties.PluginIdentifier);
 
-		if (StateProperties.Destination.MaxState == EGameFeaturePluginState::Registered)
-		{
-			// If we aren't going farther than Registered, GC now
-			// otherwise we will defer until closer to our destination state
-			bHasUnloaded = true;
-			UpdateStateMachineDeferred();
-			return;
-		}
-
 		StateStatus.SetTransition(EGameFeaturePluginState::Registered);
 	}
 
diff --git a/Engine/Plugins/Experimental/GameFeatures/Source/GameFeatures/Private/GameFeaturesSubsystem.cpp b/Engine/Plugins/Experimental/GameFeatures/Source/GameFeatures/Private/GameFeaturesSubsystem.cpp
--- a/Engine/Plugins/Experimental/GameFeatures/Source/GameFeatures/Private/GameFeaturesSubsystem.cpp
+++ b/Engine/Plugins/Experimental/GameFeatures/Source/GameFeatures/Private/GameFeaturesSubsystem.cpp
@@ -220,6 +220,11 @@ bool FGameFeaturePluginIdentifier::ExactMatchesURL(const FString& PluginURLIn) c
 	return GetFullPluginURL().Equals(PluginURLIn, ESearchCase::IgnoreCase);
 }
 
+FStringView FGameFeaturePluginIdentifier::GetPluginName() const
+{ 
+	return FPathViews::GetBaseFilename(IdentifyingURLSubset); 
+}
+
 void FGameFeatureStateChangeContext::SetRequiredWorldContextHandle(FName Handle)
 {
 	WorldContextHandle = Handle;
@@ -641,12 +646,12 @@ void UGameFeaturesSubsystem::RemoveGameFeatureFromAssetManager(const UGameFeatur
 	}
 }
 
-void UGameFeaturesSubsystem::ForEachGameFeature(TFunction<void(FGameFeatureInfo&&)>& Visitor) const
+void UGameFeaturesSubsystem::ForEachGameFeature(TFunctionRef<void(FGameFeatureInfo&&)> Visitor) const
 {
 	for (auto StateMachineIt = GameFeaturePluginStateMachines.CreateConstIterator(); StateMachineIt; ++StateMachineIt)
 	{
 		if (UGameFeaturePluginStateMachine* GFSM = StateMachineIt.Value())
-		{	
+		{
 			FGameFeatureInfo GameFeatureInfo = { GFSM->GetPluginName(), GFSM->GetPluginURL(), GFSM->WasLoadedAsBuiltIn(), GFSM->GetCurrentState() };
 			Visitor(MoveTemp(GameFeatureInfo));
 		}
@@ -835,6 +840,21 @@ void UGameFeaturesSubsystem::OnGameFeatureStatusKnown(const FString& PluginName,
 	}
 }
 
+void UGameFeaturesSubsystem::OnGameFeaturePredownloading(const FString& PluginName, const FGameFeaturePluginIdentifier& PluginIdentifier)
+{
+	CallbackObservers(EObserverCallback::Predownloading, PluginIdentifier, &PluginName);
+}
+
+void UGameFeaturesSubsystem::OnGameFeatureDownloading(const FString& PluginName, const FGameFeaturePluginIdentifier& PluginIdentifier)
+{
+	CallbackObservers(EObserverCallback::Downloading, PluginIdentifier, &PluginName);
+}
+
+void UGameFeaturesSubsystem::OnGameFeatureReleasing(const FString& PluginName, const FGameFeaturePluginIdentifier& PluginIdentifier)
+{
+	CallbackObservers(EObserverCallback::Releasing, PluginIdentifier, &PluginName);
+}
+
 void UGameFeaturesSubsystem::OnGameFeaturePreMounting(const FString& PluginName, const FGameFeaturePluginIdentifier& PluginIdentifier, FGameFeaturePreMountingContext& Context)
 {
 	CallbackObservers(EObserverCallback::PreMounting, PluginIdentifier, &PluginName, /*GameFeatureData=*/nullptr, &Context);
@@ -2185,6 +2205,8 @@ struct FGameFeaturePluginPredownloadContext : public FGameFeaturePluginPredownlo
 			}
 
 			BundlesToInstall.Append(Pair.Value.InstallBundles);
+
+			GFPSubSys.OnGameFeaturePredownloading(FString(Pair.Key.GetPluginName()), Pair.Key);
 		}
 
 		TSharedPtr<IInstallBundleManager> BundleManager = IInstallBundleManager::GetPlatformInstallBundleManager();
@@ -2776,7 +2798,7 @@ void UGameFeaturesSubsystem::CallbackObservers(EObserverCallback CallbackType, c
 	const UGameFeatureData* GameFeatureData /*= nullptr*/, 
 	FGameFeatureStateChangeContext* StateChangeContext /*= nullptr*/)
 {
-	static_assert(std::underlying_type<EObserverCallback>::type(EObserverCallback::Count) == 11, "Update UGameFeaturesSubsystem::CallbackObservers to handle added EObserverCallback");
+	static_assert(std::underlying_type<EObserverCallback>::type(EObserverCallback::Count) == 14, "Update UGameFeaturesSubsystem::CallbackObservers to handle added EObserverCallback");
 
 	// Protect against modifying the observer list during iteration
 	TArray<UObject*> LocalObservers(Observers);
@@ -2799,6 +2821,33 @@ void UGameFeaturesSubsystem::CallbackObservers(EObserverCallback CallbackType, c
 		}
 		break;
 	}
+	case EObserverCallback::Predownloading:
+	{
+		check(PluginName);
+		for (UObject* Observer : LocalObservers)
+		{
+			CastChecked<IGameFeatureStateChangeObserver>(Observer)->OnGameFeaturePredownloading(*PluginName, PluginIdentifier);
+		}
+		break;
+	}
+	case EObserverCallback::Downloading:
+	{
+		check(PluginName);
+		for (UObject* Observer : LocalObservers)
+		{
+			CastChecked<IGameFeatureStateChangeObserver>(Observer)->OnGameFeatureDownloading(*PluginName, PluginIdentifier);
+		}
+		break;
+	}
+	case EObserverCallback::Releasing:
+	{
+		check(PluginName);
+		for (UObject* Observer : LocalObservers)
+		{
+			CastChecked<IGameFeatureStateChangeObserver>(Observer)->OnGameFeatureReleasing(*PluginName, PluginIdentifier);
+		}
+		break;
+	}
 	case EObserverCallback::PreMounting:
 	{
 		check(PluginName);
diff --git a/Engine/Plugins/Experimental/GameFeatures/Source/GameFeatures/Public/GameFeatureOptionalContentInstaller.h b/Engine/Plugins/Experimental/GameFeatures/Source/GameFeatures/Public/GameFeatureOptionalContentInstaller.h
new file mode 100644
--- /dev/null
+++ b/Engine/Plugins/Experimental/GameFeatures/Source/GameFeatures/Public/GameFeatureOptionalContentInstaller.h
@@ -0,0 +1,76 @@
+// Copyright Epic Games, Inc. All Rights Reserved.
+
+#pragma once
+
+#include "GameFeatureStateChangeObserver.h"
+#include "HAL/IConsoleManager.h"
+#include "InstallBundleManagerInterface.h"
+#include "GameFeatureOptionalContentInstaller.generated.h"
+
+namespace UE::GameFeatures
+{
+	struct FResult;
+}
+
+/** 
+ * Utilty class to install GFP optional paks (usually containing optional mips) in sync with GFP content installs.
+ * NOTE: This only currently supports LRU cached install bundles. It would need UI callbacks and additional support 
+ * for free space checks and progress tracking to fully support non-LRU GFPs.
+ */
+UCLASS(MinimalAPI)
+class UGameFeatureOptionalContentInstaller : public UObject, public IGameFeatureStateChangeObserver
+{
+	GENERATED_BODY()
+
+public:
+	static GAMEFEATURES_API TMulticastDelegate<void(const FString& PluginName, const UE::GameFeatures::FResult&)> OnOptionalContentInstalled;
+
+public:
+	virtual void BeginDestroy() override;
+
+	void GAMEFEATURES_API Init(TUniqueFunction<TArray<FName>(FString)> GetOptionalBundlePredicate);
+
+	void GAMEFEATURES_API Enable(bool bEnable);
+
+	void GAMEFEATURES_API EnableCellularDownloading(bool bEnable);
+
+private:
+	bool UpdateContent(const FString& PluginName, bool bIsPredownload);
+
+	void OnContentInstalled(FInstallBundleRequestResultInfo InResult, FString PluginName);
+
+	void ReleaseContent(const FString& PluginName);
+
+	void OnEnabled();
+	void OnDisabled();
+
+	bool IsEnabled() const;
+
+	void OnCVarsChanged();
+
+	// IGameFeatureStateChangeObserver Interface
+	virtual void OnGameFeaturePredownloading(const FString& PluginName, const FGameFeaturePluginIdentifier& PluginIdentifier) override;
+	virtual void OnGameFeatureDownloading(const FString& PluginName, const FGameFeaturePluginIdentifier& PluginIdentifier) override;
+	virtual void OnGameFeatureReleasing(const FString& PluginName, const FGameFeaturePluginIdentifier& PluginIdentifier) override;
+
+private:
+	struct FGFPInstall
+	{
+		FDelegateHandle CallbackHandle;
+		TArray<FName> BundlesEnqueued;
+		bool bIsPredownload = false;
+	};
+
+private:
+	TUniqueFunction<TArray<FName>(FString)> GetOptionalBundlePredicate;
+	TSharedPtr<IInstallBundleManager> BundleManager;
+	TSet<FString> RelevantGFPs;
+	TMap<FString, FGFPInstall> ActiveGFPInstalls;
+
+	/** Delegate handle for a consolve variable sink */
+	FConsoleVariableSinkHandle CVarSinkHandle;
+
+	bool bEnabled = false;
+	bool bEnabledCVar = false;
+	bool bAllowCellDownload = false;
+};
diff --git a/Engine/Plugins/Experimental/GameFeatures/Source/GameFeatures/Public/GameFeatureStateChangeObserver.h b/Engine/Plugins/Experimental/GameFeatures/Source/GameFeatures/Public/GameFeatureStateChangeObserver.h
--- a/Engine/Plugins/Experimental/GameFeatures/Source/GameFeatures/Public/GameFeatureStateChangeObserver.h
+++ b/Engine/Plugins/Experimental/GameFeatures/Source/GameFeatures/Public/GameFeatureStateChangeObserver.h
@@ -40,6 +40,15 @@ public:
 	// Invoked prior to terminating a game feature plugin
 	virtual void OnGameFeatureTerminating(const FString& PluginURL) {}
 
+	// Invoked when content begins installing via predownload
+	virtual void OnGameFeaturePredownloading(const FString& PluginName, const FGameFeaturePluginIdentifier& PluginIdentifier) {}
+
+	// Invoked when content begins installing
+	virtual void OnGameFeatureDownloading(const FString& PluginName, const FGameFeaturePluginIdentifier& PluginIdentifier) {}
+
+	// Invoked when content is released (the point it at which it is safe to remove it)
+	virtual void OnGameFeatureReleasing(const FString& PluginName, const FGameFeaturePluginIdentifier& PluginIdentifier) {}
+
 	// Invoked prior to mounting a plugin (but after its install bundles become available, if any)
 	virtual void OnGameFeaturePreMounting(const FString& PluginName, const FGameFeaturePluginIdentifier& PluginIdentifier, FGameFeaturePreMountingContext& Context) {}
 
diff --git a/Engine/Plugins/Experimental/GameFeatures/Source/GameFeatures/Public/GameFeaturesSubsystem.h b/Engine/Plugins/Experimental/GameFeatures/Source/GameFeatures/Public/GameFeaturesSubsystem.h
--- a/Engine/Plugins/Experimental/GameFeatures/Source/GameFeatures/Public/GameFeaturesSubsystem.h
+++ b/Engine/Plugins/Experimental/GameFeatures/Source/GameFeatures/Public/GameFeaturesSubsystem.h
@@ -318,6 +318,9 @@ struct GAMEFEATURES_API FGameFeaturePluginIdentifier
 	/** Returns the Identifying information used for this Plugin. It is a subset of the URL used to create it.*/
 	FStringView GetIdentifyingString() const { return IdentifyingURLSubset; }
 
+	/** Returns the name of the plugin */
+	FStringView GetPluginName() const;
+
 	/** Get the Full PluginURL used to originally construct this identifier */
 	const FString& GetFullPluginURL() const { return PluginURL; }
 
@@ -416,7 +419,7 @@ public:
 	void AddObserver(UObject* Observer);
 	void RemoveObserver(UObject* Observer);
 
-	void ForEachGameFeature(TFunction<void(FGameFeatureInfo&&)>& Visitor) const;
+	void ForEachGameFeature(TFunctionRef<void(FGameFeatureInfo&&)> Visitor) const;
 
 	/**
 	 * Calls the compile-time lambda on each active game feature data of the specified type
@@ -638,6 +641,15 @@ private:
 	void OnGameFeatureStatusKnown(const FString& PluginName, const FGameFeaturePluginIdentifier& PluginIdentifierL);
 	friend struct FGameFeaturePluginState_CheckingStatus;
 
+	void OnGameFeaturePredownloading(const FString& PluginName, const FGameFeaturePluginIdentifier& PluginIdentifier);
+
+	void OnGameFeatureDownloading(const FString& PluginName, const FGameFeaturePluginIdentifier& PluginIdentifier);
+	friend struct FGameFeaturePluginState_Downloading;
+
+	void OnGameFeatureReleasing(const FString& PluginName, const FGameFeaturePluginIdentifier& PluginIdentifier);
+	friend struct FGameFeaturePluginState_Releasing;
+	friend struct FGameFeaturePluginState_Unmounting;
+
 	void OnGameFeaturePreMounting(const FString& PluginName, const FGameFeaturePluginIdentifier& PluginIdentifier, FGameFeaturePreMountingContext& Context);
 	void OnGameFeaturePostMounting(const FString& PluginName, const FGameFeaturePluginIdentifier& PluginIdentifier, FGameFeaturePostMountingContext& Context);
 	friend struct FGameFeaturePluginState_Mounting;
@@ -735,6 +747,9 @@ private:
 	{
 		CheckingStatus,
 		Terminating,
+		Predownloading,
+		Downloading,
+		Releasing,
 		PreMounting,
 		PostMounting,
 		Registering,
