Enable seperate installation of optional GFP content.

diff --git a/Engine/Plugins/Experimental/GameFeatures/Source/GameFeatures/Private/GameFeatureOptionalContentInstaller.cpp b/Engine/Plugins/Experimental/GameFeatures/Source/GameFeatures/Private/GameFeatureOptionalContentInstaller.cpp
--- a/Engine/Plugins/Experimental/GameFeatures/Source/GameFeatures/Private/GameFeatureOptionalContentInstaller.cpp
+++ b/Engine/Plugins/Experimental/GameFeatures/Source/GameFeatures/Private/GameFeatureOptionalContentInstaller.cpp
@@ -17,7 +17,7 @@ namespace GameFeatureOptionalContentInstaller
 	static const FStringView ErrorNamespace = TEXTVIEW("GameFeaturePlugin.OptionalDownload.");
 
 	static TAutoConsoleVariable<bool> CVarEnableOptionalContentInstaller(TEXT("GameFeatureOptionalContentInstaller.Enable"), 
-		false,
+		true,
 		TEXT("Enable optional content installer"));
 }
 
