Fix the Game Feature Data UI to show an error instead of infinitely ensure if you make one outside of a game feature plugin

diff --git a/Engine/Plugins/Experimental/GameFeatures/Source/GameFeaturesEditor/Private/GameFeatureDataDetailsCustomization.cpp b/Engine/Plugins/Experimental/GameFeatures/Source/GameFeaturesEditor/Private/GameFeatureDataDetailsCustomization.cpp
--- a/Engine/Plugins/Experimental/GameFeatures/Source/GameFeaturesEditor/Private/GameFeatureDataDetailsCustomization.cpp
+++ b/Engine/Plugins/Experimental/GameFeatures/Source/GameFeaturesEditor/Private/GameFeatureDataDetailsCustomization.cpp
@@ -214,6 +214,10 @@ void FGameFeatureDataDetailsCustomization::ChangeDesiredState(EGameFeaturePlugin
 
 EGameFeaturePluginState FGameFeatureDataDetailsCustomization::GetCurrentState() const
 {
+	if (PluginURL.IsEmpty())
+	{
+		return EGameFeaturePluginState::Uninitialized;
+	}
 	return UGameFeaturesSubsystem::Get().GetPluginState(PluginURL);
 }
 
@@ -231,6 +235,11 @@ FText FGameFeatureDataDetailsCustomization::GetInitialStateText() const
 
 FText FGameFeatureDataDetailsCustomization::GetTagConfigPathText() const
 {
+	if (PluginURL.IsEmpty())
+	{
+		return LOCTEXT("TagConfigPathInvalid", "Invalid Plugin");
+	}
+
 	FString PluginFile = UGameFeaturesSubsystem::Get().GetPluginFilenameFromPluginURL(PluginURL);
 	FString PluginFolder = FPaths::GetPath(PluginFile);
 	FString TagFolder = PluginFolder / TEXT("Config") / TEXT("Tags");
