Fix an issue where primary asset override rules would not apply during the cook.

diff --git a/Engine/Plugins/Experimental/GameFeatures/Source/GameFeatures/Private/GameFeaturesSubsystem.cpp b/Engine/Plugins/Experimental/GameFeatures/Source/GameFeatures/Private/GameFeaturesSubsystem.cpp
--- a/Engine/Plugins/Experimental/GameFeatures/Source/GameFeatures/Private/GameFeaturesSubsystem.cpp
+++ b/Engine/Plugins/Experimental/GameFeatures/Source/GameFeatures/Private/GameFeaturesSubsystem.cpp
@@ -613,8 +613,16 @@ void UGameFeaturesSubsystem::AddGameFeatureToAssetManager(const UGameFeatureData
 			}
 		}
 	}
-
 	LocalAssetManager.PopBulkScanning();
+
+	const UAssetManagerSettings& Settings = LocalAssetManager.GetSettings();
+	for (const FPrimaryAssetRulesCustomOverride& Override : Settings.CustomPrimaryAssetRules)
+	{
+		if (Override.FilterDirectory.Path.StartsWith(PluginRootPath))
+		{
+			LocalAssetManager.ApplyCustomPrimaryAssetRulesOverride(Override);
+		}
+	}
 }
 
 void UGameFeaturesSubsystem::RemoveGameFeatureFromAssetManager(const UGameFeatureData* GameFeatureToRemove, const FString& PluginName, const TArray<FName>& AddedPrimaryAssetTypes)
