Unshelved from pending changelist '31032247':

Moving Game Feature action section to the top of a GFD, since it is the most commonly used field.

diff --git a/Engine/Plugins/Experimental/GameFeatures/Source/GameFeatures/Public/GameFeatureData.h b/Engine/Plugins/Experimental/GameFeatures/Source/GameFeatures/Public/GameFeatureData.h
--- a/Engine/Plugins/Experimental/GameFeatures/Source/GameFeatures/Public/GameFeatureData.h
+++ b/Engine/Plugins/Experimental/GameFeatures/Source/GameFeatures/Public/GameFeatureData.h
@@ -98,7 +98,7 @@ private:
 protected:
 
 	/** List of actions to perform as this game feature is loaded/activated/deactivated/unloaded */
-	UPROPERTY(EditDefaultsOnly, Instanced, Category="Actions", meta = (GetDisallowedClasses = "GetDisallowedActions"))
+	UPROPERTY(EditDefaultsOnly, Instanced, Category="Game Feature | Actions", meta = (GetDisallowedClasses = "GetDisallowedActions"))
 	TArray<TObjectPtr<UGameFeatureAction>> Actions;
 
 	/** List of asset types to scan at startup */
