Update to support Empty FilterDirectory paths.

diff --git a/Engine/Plugins/Experimental/GameFeatures/Source/GameFeatures/Private/GameFeaturesSubsystem.cpp b/Engine/Plugins/Experimental/GameFeatures/Source/GameFeatures/Private/GameFeaturesSubsystem.cpp
--- a/Engine/Plugins/Experimental/GameFeatures/Source/GameFeatures/Private/GameFeaturesSubsystem.cpp
+++ b/Engine/Plugins/Experimental/GameFeatures/Source/GameFeatures/Private/GameFeaturesSubsystem.cpp
@@ -618,7 +618,7 @@ void UGameFeaturesSubsystem::AddGameFeatureToAssetManager(const UGameFeatureData
 	const UAssetManagerSettings& Settings = LocalAssetManager.GetSettings();
 	for (const FPrimaryAssetRulesCustomOverride& Override : Settings.CustomPrimaryAssetRules)
 	{
-		if (Override.FilterDirectory.Path.StartsWith(PluginRootPath))
+		if (Override.FilterDirectory.Path.IsEmpty() || Override.FilterDirectory.Path.StartsWith(PluginRootPath))
 		{
 			LocalAssetManager.ApplyCustomPrimaryAssetRulesOverride(Override);
 		}
