Fix comments

#changelist validated
#virtualized

diff --git a/Engine/Plugins/Runtime/GameFeatures/Source/GameFeatures/Private/GameFeatureOptionalContentInstaller.cpp b/Engine/Plugins/Runtime/GameFeatures/Source/GameFeatures/Private/GameFeatureOptionalContentInstaller.cpp
--- a/Engine/Plugins/Runtime/GameFeatures/Source/GameFeatures/Private/GameFeatureOptionalContentInstaller.cpp
+++ b/Engine/Plugins/Runtime/GameFeatures/Source/GameFeatures/Private/GameFeatureOptionalContentInstaller.cpp
@@ -34,7 +34,7 @@ void UGameFeatureOptionalContentInstaller::Init(TUniqueFunction<TArray<FName>(FS
 	GetOptionalBundlePredicate = MoveTemp(InGetOptionalBundlePredicate);
 	BundleManager = IInstallBundleManager::GetPlatformInstallBundleManager();
 
-	// Create the cvar sink and set up the thread pool
+	// Create the cvar sink
 	CVarSinkHandle = IConsoleManager::Get().RegisterConsoleVariableSink_Handle(
 		FConsoleCommandDelegate::CreateUObject(this, &UGameFeatureOptionalContentInstaller::OnCVarsChanged));
 	bEnabledCVar = GameFeatureOptionalContentInstaller::CVarEnableOptionalContentInstaller.GetValueOnGameThread();
diff --git a/Engine/Plugins/Runtime/GameFeatures/Source/GameFeatures/Public/GameFeatureOptionalContentInstaller.h b/Engine/Plugins/Runtime/GameFeatures/Source/GameFeatures/Public/GameFeatureOptionalContentInstaller.h
--- a/Engine/Plugins/Runtime/GameFeatures/Source/GameFeatures/Public/GameFeatureOptionalContentInstaller.h
+++ b/Engine/Plugins/Runtime/GameFeatures/Source/GameFeatures/Public/GameFeatureOptionalContentInstaller.h
@@ -67,7 +67,7 @@ private:
 	TSet<FString> RelevantGFPs;
 	TMap<FString, FGFPInstall> ActiveGFPInstalls;
 
-	/** Delegate handle for a consolve variable sink */
+	/** Delegate handle for a console variable sink */
 	FConsoleVariableSinkHandle CVarSinkHandle;
 
 	bool bEnabled = false;
