PR #9012: Ensure only plugins in the GameFeatures/ folder are considered as GameFeatures.

diff --git a/Engine/Plugins/Runtime/GameFeatures/Source/GameFeatures/Private/GameFeaturesSubsystemSettings.cpp b/Engine/Plugins/Runtime/GameFeatures/Source/GameFeatures/Private/GameFeaturesSubsystemSettings.cpp
--- a/Engine/Plugins/Runtime/GameFeatures/Source/GameFeatures/Private/GameFeaturesSubsystemSettings.cpp
+++ b/Engine/Plugins/Runtime/GameFeatures/Source/GameFeatures/Private/GameFeaturesSubsystemSettings.cpp
@@ -34,7 +34,7 @@ bool UGameFeaturesSubsystemSettings::IsValidGameFeaturePlugin(const FString& Plu
 
 		for (FString& BuiltInFolder : BuiltInGameFeaturePluginsFolders)
 		{
-			BuiltInFolder = FPaths::ConvertRelativePathToFull(BuiltInFolder);
+			BuiltInFolder = FPaths::ConvertRelativePathToFull(BuiltInFolder + TEXT("/"));
 		}
 	}
 
