Add validator so nonGameFeaturePlugins do not depend on GameFeaturePlugins

#presubmit 65cfea1b05690e8376a10025

diff --git a/Engine/Plugins/Runtime/GameFeatures/Source/GameFeatures/GameFeatures.Build.cs b/Engine/Plugins/Runtime/GameFeatures/Source/GameFeatures/GameFeatures.Build.cs
--- a/Engine/Plugins/Runtime/GameFeatures/Source/GameFeatures/GameFeatures.Build.cs
+++ b/Engine/Plugins/Runtime/GameFeatures/Source/GameFeatures/GameFeatures.Build.cs
@@ -40,6 +40,7 @@ namespace UnrealBuildTool.Rules
 					{
 						"UnrealEd",
 						"PluginUtils",
+						"DataValidation"
 					}
 				);
 			}
diff --git a/Engine/Plugins/Runtime/GameFeatures/Source/GameFeatures/Private/IllegalPluginDependenciesValidator.cpp b/Engine/Plugins/Runtime/GameFeatures/Source/GameFeatures/Private/IllegalPluginDependenciesValidator.cpp
new file mode 100644
--- /dev/null
+++ b/Engine/Plugins/Runtime/GameFeatures/Source/GameFeatures/Private/IllegalPluginDependenciesValidator.cpp
@@ -0,0 +1,114 @@
+// Copyright Epic Games, Inc. All Rights Reserved.
+
+#if WITH_EDITOR
+#include "IllegalPluginDependenciesValidator.h"
+
+#include "DataValidationChangelist.h"
+#include "GameFeaturesSubsystem.h"
+#include "Interfaces/IPluginManager.h"
+
+#include UE_INLINE_GENERATED_CPP_BY_NAME(IllegalPluginDependenciesValidator)
+
+#define LOCTEXT_NAMESPACE "IllegalPluginDependenciesValidator"
+
+namespace IllegalPluginDependencies
+{
+	bool IsGameFeaturePlugin(const TSharedPtr<IPlugin> InPlugin)
+	{
+		const FString PluginsFolderRoot = (InPlugin->GetLoadedFrom() == EPluginLoadedFrom::Project) ? FPaths::ProjectPluginsDir() : FPaths::EnginePluginsDir();
+		FString PluginPathRelativeToDomain = InPlugin->GetBaseDir();
+		if (FPaths::MakePathRelativeTo(PluginPathRelativeToDomain, *PluginsFolderRoot))
+		{
+			PluginPathRelativeToDomain = TEXT("/") + PluginPathRelativeToDomain + TEXT("/");
+		}
+		else
+		{
+			PluginPathRelativeToDomain = FString();
+		}
+
+		return (PluginPathRelativeToDomain.StartsWith(TEXT("/GameFeatures/")));
+	}
+}
+
+bool UIllegalPluginDependenciesValidator::CanValidateAsset_Implementation(const FAssetData& AssetData, UObject* InAsset, FDataValidationContext& InContext) const
+{
+	return (InAsset->GetClass() == UDataValidationChangelist::StaticClass());
+}
+
+EDataValidationResult UIllegalPluginDependenciesValidator::ValidateLoadedAsset_Implementation(const FAssetData& AssetData, UObject* InAsset, FDataValidationContext& InContext)
+{
+	UDataValidationChangelist* DataValidationChangelist = CastChecked<UDataValidationChangelist>(InAsset);
+
+	bool bChangelistContainsPluginFiles = false;
+	for (const FString& ModifiedFile : DataValidationChangelist->ModifiedFiles)
+	{
+		if (ModifiedFile.EndsWith(TEXT(".uplugin")))
+		{
+			bChangelistContainsPluginFiles = true;
+			break;
+		}
+	}
+
+	// Early out if the changelist doesn't contain any uplugin files.
+	if (!bChangelistContainsPluginFiles)
+	{
+		AssetPasses(InAsset);
+		return GetValidationResult();
+	}
+
+	TSet<FString> GFPs;
+	TArray<TSharedRef<IPlugin>> AllPlugins = IPluginManager::Get().GetDiscoveredPlugins();
+	GFPs.Reserve(AllPlugins.Num());
+	for (const TSharedRef<IPlugin>& Plugin : AllPlugins)
+	{
+		FString PluginURL;
+		FGameFeaturePluginDetails PluginDetails;
+		if (UGameFeaturesSubsystem::Get().GetBuiltInGameFeaturePluginDetails(Plugin, PluginURL, PluginDetails))
+		{
+			GFPs.Add(Plugin->GetName());
+		}
+	}
+
+	for (const FString& ModifiedFile : DataValidationChangelist->ModifiedFiles)
+	{
+		if (ModifiedFile.EndsWith(TEXT(".uplugin")))
+		{
+			FString ModifiedPlugin = FPaths::GetBaseFilename(ModifiedFile);
+			TSharedPtr<IPlugin> Plugin = IPluginManager::Get().FindPlugin(ModifiedPlugin);
+			if (Plugin.IsValid())
+			{
+				if (GFPs.Contains(ModifiedPlugin) || IllegalPluginDependencies::IsGameFeaturePlugin(Plugin))
+				{
+					continue;
+				}
+
+				// Plugin is not a GFP so it can not have any GFP dependencies
+				const FPluginDescriptor& Descriptor = Plugin->GetDescriptor();
+				for (const FPluginReferenceDescriptor& Dependency : Descriptor.Plugins)
+				{
+					TSharedPtr<IPlugin> DependencyPlugin = IPluginManager::Get().FindPlugin(Dependency.Name);
+					if (DependencyPlugin.IsValid())
+					{
+						if (GFPs.Contains(Dependency.Name))
+						{
+							FText NewError = FText::Format(
+								LOCTEXT("ValidationError.IllegalPluginDependency", "Plugin {0} depends on {1}. Non GameFeaturePlugins are not allowed to depend on GameFeaturePlugins. This can create an issue where objects will fail to load"),
+								FText::FromString(ModifiedPlugin),
+								FText::FromString(Dependency.Name)
+							);
+							AssetFails(InAsset, NewError);
+						}
+					}
+				}
+			}
+		}
+	}
+
+	if (GetValidationResult() != EDataValidationResult::Invalid)
+	{
+		AssetPasses(InAsset);
+	}
+
+	return GetValidationResult();
+}
+#endif
\ No newline at end of file
diff --git a/Engine/Plugins/Runtime/GameFeatures/Source/GameFeatures/Private/IllegalPluginDependenciesValidator.h b/Engine/Plugins/Runtime/GameFeatures/Source/GameFeatures/Private/IllegalPluginDependenciesValidator.h
new file mode 100644
--- /dev/null
+++ b/Engine/Plugins/Runtime/GameFeatures/Source/GameFeatures/Private/IllegalPluginDependenciesValidator.h
@@ -0,0 +1,25 @@
+// Copyright Epic Games, Inc. All Rights Reserved.
+
+#pragma once
+
+#if WITH_EDITOR
+#include "EditorValidatorBase.h"
+#include "IllegalPluginDependenciesValidator.generated.h"
+
+class FText;
+class UObject;
+
+/**
+ * Ensures that non-GameFeaturePlugins do not depend on GameFeaturePlugins.
+ * GameFeaturePlugins will load content later than non-GameFeaturePlugins which could cause linker load issues if they do not exist.
+ */
+UCLASS()
+class UIllegalPluginDependenciesValidator : public UEditorValidatorBase
+{
+	GENERATED_BODY()
+
+protected:
+	virtual bool CanValidateAsset_Implementation(const FAssetData& InAssetData, UObject* InAsset, FDataValidationContext& InContext) const override;
+	virtual EDataValidationResult ValidateLoadedAsset_Implementation(const FAssetData& InAssetData, UObject* InAsset, FDataValidationContext& InContext) override;
+};
+#endif
\ No newline at end of file
