Fix plugin dependency.
Add DataValidation for editor targets

#presubmit 65cfea1b05690e8376a10025

diff --git a/Engine/Plugins/Runtime/GameFeatures/GameFeatures.uplugin b/Engine/Plugins/Runtime/GameFeatures/GameFeatures.uplugin
--- a/Engine/Plugins/Runtime/GameFeatures/GameFeatures.uplugin
+++ b/Engine/Plugins/Runtime/GameFeatures/GameFeatures.uplugin
@@ -43,6 +43,11 @@
 			"Name": "PluginUtils",
 			"Enabled": true,
 			"TargetAllowList": [ "Editor" ]
+		},
+		{
+			"Name": "DataValidation",
+			"Enabled": true,
+			"TargetAllowList": [ "Editor" ]
 		}
 	]
 }
