Undo single changelist 31586002

diff --git a/Engine/Plugins/Runtime/GameFeatures/Source/GameFeatures/GameFeatures.Build.cs b/Engine/Plugins/Runtime/GameFeatures/Source/GameFeatures/GameFeatures.Build.cs
--- a/Engine/Plugins/Runtime/GameFeatures/Source/GameFeatures/GameFeatures.Build.cs
+++ b/Engine/Plugins/Runtime/GameFeatures/Source/GameFeatures/GameFeatures.Build.cs
@@ -40,7 +40,6 @@ namespace UnrealBuildTool.Rules
 					{
 						"UnrealEd",
 						"PluginUtils",
-						"DataValidation"
 					}
 				);
 			}
diff --git a/Engine/Plugins/Runtime/GameFeatures/Source/GameFeatures/Private/IllegalPluginDependenciesValidator.cpp b/Engine/Plugins/Runtime/GameFeatures/Source/GameFeatures/Private/IllegalPluginDependenciesValidator.cpp
deleted file mode 100644
diff --git a/Engine/Plugins/Runtime/GameFeatures/Source/GameFeatures/Private/IllegalPluginDependenciesValidator.h b/Engine/Plugins/Runtime/GameFeatures/Source/GameFeatures/Private/IllegalPluginDependenciesValidator.h
deleted file mode 100644
