GFP Subsystem
- Utility functions for parsing GFP URL Options
- Move URL Options enum into GameFeatureTypes.h
- Cleanup redundant ToString and LexToString functions
- Moved some utilities into the modules that actually used them

diff --git a/Engine/Plugins/Runtime/GameFeatures/Source/GameFeatures/Private/GameFeaturePluginStateMachine.cpp b/Engine/Plugins/Runtime/GameFeatures/Source/GameFeatures/Private/GameFeaturePluginStateMachine.cpp
--- a/Engine/Plugins/Runtime/GameFeatures/Source/GameFeatures/Private/GameFeaturePluginStateMachine.cpp
+++ b/Engine/Plugins/Runtime/GameFeatures/Source/GameFeatures/Private/GameFeaturePluginStateMachine.cpp
@@ -81,22 +81,8 @@ namespace UE::GameFeatures
 		false,
 		TEXT("Enable to force calls to IAssetRegistry::AppendState to happen on the game thread"));
 
-	#define GAME_FEATURE_PLUGIN_STATE_TO_STRING(inEnum, inText) case EGameFeaturePluginState::inEnum: return TEXT(#inEnum);
-	FString ToString(EGameFeaturePluginState InType)
-	{
-		switch (InType)
-		{
-		GAME_FEATURE_PLUGIN_STATE_LIST(GAME_FEATURE_PLUGIN_STATE_TO_STRING)
-		default:
-			check(0);
-			return FString();
-		}
-	}
-	#undef GAME_FEATURE_PLUGIN_STATE_TO_STRING
-
 	bool ShouldSkipVerify(const FString& PluginName)
 	{
-
 		static const FAsciiSet Wildcards("*?");
 		bool bSkip = false;
 		UE::String::ParseTokens(MakeStringView(GVerifyPluginSkipList), TEXTVIEW(","), [&PluginName, &bSkip](FStringView Item) {
@@ -138,19 +124,6 @@ namespace UE::GameFeatures
 		return Priority;
 	}
 
-#define GAME_FEATURE_PLUGIN_PROTOCOL_PREFIX(inEnum, inString) case EGameFeaturePluginProtocol::inEnum: return inString;
-	const TCHAR* GameFeaturePluginProtocolPrefix(EGameFeaturePluginProtocol Protocol)
-	{
-		switch (Protocol)
-		{
-			GAME_FEATURE_PLUGIN_PROTOCOL_LIST(GAME_FEATURE_PLUGIN_PROTOCOL_PREFIX)
-		}
-
-		check(false);
-		return nullptr;
-	}
-#undef GAME_FEATURE_PLUGIN_PROTOCOL_PREFIX
-
 	static bool bRealtimeMode = false;
 
 	class FRealtimeMode : public TSharedFromThis<FRealtimeMode>
@@ -291,42 +264,6 @@ namespace UE::GameFeatures
 #endif //if WITH_EDITOR
 }
 
-//Enum describing the options for what data can be inside the InstallBundle protocol URL's metadata
-//Used to aid in parsing / creating the URL for InstallBundle protocol GameFeaturePlugins
-#define INSTALL_BUNDLE_PROTOCOL_URL_OPTIONS_LIST(XOPTION)					\
-	XOPTION(Bundles,					TEXT("Bundles"))					
-
-#define INSTALL_BUNDLE_PROTOCOL_URL_OPTIONS_ENUM(inEnum, inString) inEnum,
-enum class EGameFeatureInstallBundleProtocolOptions : uint8
-{
-	INSTALL_BUNDLE_PROTOCOL_URL_OPTIONS_LIST(INSTALL_BUNDLE_PROTOCOL_URL_OPTIONS_ENUM)
-	Count,
-};
-ENUM_RANGE_BY_COUNT(EGameFeatureInstallBundleProtocolOptions, EGameFeatureInstallBundleProtocolOptions::Count);
-#undef INSTALL_BUNDLE_PROTOCOL_URL_OPTIONS_ENUM
-
-#define INSTALL_BUNDLE_PROTOCOL_URL_OPTIONS_LEX_TO_STRING(inEnum, inString) \
-    case(EGameFeatureInstallBundleProtocolOptions::inEnum):                 \
-    {                                                                       \
-        return inString;                                                    \
-    }                                                                       
-
-FString LexToString(EGameFeatureInstallBundleProtocolOptions FeatureIn)
-{
-	switch (FeatureIn)
-	{
-		INSTALL_BUNDLE_PROTOCOL_URL_OPTIONS_LIST(INSTALL_BUNDLE_PROTOCOL_URL_OPTIONS_LEX_TO_STRING)
-		
-		default:
-		{
-			ensureAlwaysMsgf(false, TEXT("Logic error causing a missing LexToString value for EGameFeatureInstallBundleProtocolOptions:%d"), static_cast<uint8>(FeatureIn));
-			return ("ERROR_UNSUPPORTED_ENUM");
-		}
-	}
-}
-#undef INSTALL_BUNDLE_PROTOCOL_URL_OPTIONS_LEX_TO_STRING
-#undef INSTALL_BUNDLE_PROTOCOL_URL_OPTIONS_LIST
-
 void FGameFeaturePluginStateStatus::SetTransition(EGameFeaturePluginState InTransitionToState)
 {
 	TransitionToState = InTransitionToState;
@@ -337,7 +274,7 @@ void FGameFeaturePluginStateStatus::SetTransition(EGameFeaturePluginState InTran
 void FGameFeaturePluginStateStatus::SetTransitionError(EGameFeaturePluginState TransitionToErrorState, UE::GameFeatures::FResult TransitionResultIn)
 {
 	TransitionToState = TransitionToErrorState;
-	if (ensureAlwaysMsgf(TransitionResultIn.HasError(), TEXT("Invalid call to SetTransitionError with an FResult that isn't an error! TransitionToErrorState: %s"), *LexToString(TransitionToErrorState)))
+	if (ensureAlwaysMsgf(TransitionResultIn.HasError(), TEXT("Invalid call to SetTransitionError with an FResult that isn't an error! TransitionToErrorState: %s"), *UE::GameFeatures::ToString(TransitionToErrorState)))
 	{
 		TransitionResult = MoveTemp(TransitionResultIn);
 	}
@@ -348,27 +285,6 @@ void FGameFeaturePluginStateStatus::SetTransitionError(EGameFeaturePluginState T
 	}
 }
 
-void LexFromString(EGameFeatureInstallBundleProtocolOptions& ValueOut, const FStringView StringIn)
-{
-	//Default value if parsing fails
-	ValueOut = EGameFeatureInstallBundleProtocolOptions::Count;
-
-	if (StringIn.IsEmpty() || StringIn[0] == TEXT('\0'))
-	{
-		UE_LOG(LogGameFeatures, Error, TEXT("Invalid empty string used for EGameFeatureInstallBundleProtocolOptions LexFromString!"));
-		return;
-	}
-
-	for (EGameFeatureInstallBundleProtocolOptions OptionToCheck : TEnumRange<EGameFeatureInstallBundleProtocolOptions>())
-	{
-		if (LexToString(OptionToCheck) == StringIn)
-		{
-			ValueOut = OptionToCheck;
-			return;
-		}
-	}
-}
-
 UE::GameFeatures::FResult FGameFeaturePluginState::GetErrorResult(const FString& ErrorCode, const FText OptionalErrorText/*= FText()*/) const
 {
 	return GetErrorResult(TEXT(""), ErrorCode, OptionalErrorText);
@@ -376,7 +292,7 @@ UE::GameFeatures::FResult FGameFeaturePluginState::GetErrorResult(const FString&
 
 UE::GameFeatures::FResult FGameFeaturePluginState::GetErrorResult(const FString& ErrorNamespaceAddition, const FString& ErrorCode, const FText OptionalErrorText/*= FText()*/) const
 {
-	const FString StateName = LexToString(UGameFeaturesSubsystem::Get().GetPluginState(StateProperties.PluginIdentifier));
+	const FString StateName = UE::GameFeatures::ToString(UGameFeaturesSubsystem::Get().GetPluginState(StateProperties.PluginIdentifier));
 	const FString ErrorCodeEnding = ErrorNamespaceAddition.IsEmpty() ? ErrorCode : ErrorNamespaceAddition + ErrorCode;
 	const FString CompleteErrorCode = FString::Printf(TEXT("%s%s.%s"), *UE::GameFeatures::StateMachineErrorNamespace, *StateName, *ErrorCodeEnding);
 	return UE::GameFeatures::FResult(MakeError(CompleteErrorCode), OptionalErrorText);
@@ -1439,7 +1355,7 @@ struct FGameFeaturePluginState_Downloading : public FGameFeaturePluginState
 
 	void NotifyPauseChange(bool bIsPaused, FString PauseReason)
 	{
-		FGameFeaturePauseStateChangeContext Context(LexToString(EGameFeaturePluginState::Downloading), PauseReason, bIsPaused);
+		FGameFeaturePauseStateChangeContext Context(UE::GameFeatures::ToString(EGameFeaturePluginState::Downloading), PauseReason, bIsPaused);
 		UGameFeaturesSubsystem::Get().OnGameFeaturePauseChange(StateProperties.PluginIdentifier, StateProperties.PluginName, Context);
 	}
 
@@ -2843,7 +2759,7 @@ struct FGameFeaturePluginState_Deactivating : public FGameFeaturePluginState
 			// Since we are pausing work during this deactivation, also notify the OnGameFeaturePauseChange delegate
 			if (NumExpectedPausers > 0)
 			{
-				FGameFeaturePauseStateChangeContext PauseContext(LexToString(EGameFeaturePluginState::Deactivating), TEXT("PendingDeactivationCallbacks"), true);
+				FGameFeaturePauseStateChangeContext PauseContext(UE::GameFeatures::ToString(EGameFeaturePluginState::Deactivating), TEXT("PendingDeactivationCallbacks"), true);
 				UGameFeaturesSubsystem::Get().OnGameFeaturePauseChange(StateProperties.PluginIdentifier, StateProperties.PluginName, PauseContext);
 			}
 		}
@@ -2853,7 +2769,7 @@ struct FGameFeaturePluginState_Deactivating : public FGameFeaturePluginState
 			//If we previously sent an OnGameFeaturePauseChange delegate we need to send that work is now unpaused
 			if (NumExpectedPausers > 0)
 			{
-				FGameFeaturePauseStateChangeContext PauseContext(LexToString(EGameFeaturePluginState::Deactivating), TEXT(""), false);
+				FGameFeaturePauseStateChangeContext PauseContext(UE::GameFeatures::ToString(EGameFeaturePluginState::Deactivating), TEXT(""), false);
 				UGameFeaturesSubsystem::Get().OnGameFeaturePauseChange(StateProperties.PluginIdentifier, StateProperties.PluginName, PauseContext);
 			}
 
@@ -3510,10 +3426,11 @@ FString FInstallBundlePluginProtocolMetaData::ToString() const
 	FString ReturnedString;
 
 	//Always encode InstallBundles
-	ReturnedString = LexToString(EGameFeatureInstallBundleProtocolOptions::Bundles) + UE::GameFeatures::PluginURLStructureInfo::OptionAssignOperator;
+	ReturnedString = FString(UE::GameFeatures::PluginURLStructureInfo::OptionSeperator) +
+		LexToString(EGameFeatureURLOptions::Bundles) + UE::GameFeatures::PluginURLStructureInfo::OptionAssignOperator;
 
 	FNameBuilder NameBuilder;
-	const FString BundlesList = FString::JoinBy(InstallBundles, UE::GameFeatures::PluginURLStructureInfo::OptionListSeperator, 
+	const FString BundlesList = FString::JoinBy(InstallBundles, UE::GameFeatures::PluginURLStructureInfo::OptionListSeperator,
 	[&NameBuilder](const FName& BundleName)
 	{
 		BundleName.ToString(NameBuilder);
@@ -3521,88 +3438,31 @@ FString FInstallBundlePluginProtocolMetaData::ToString() const
 	});
 	ReturnedString.Append(BundlesList);
 
-	static_assert(static_cast<uint8>(EGameFeatureInstallBundleProtocolOptions::Count) == 1, "Update this function to handle the newly added EGameFeatureInstallBundleProtocolOptions value!");
+	static_assert(FMath::CountBits(static_cast<uint64>(EGameFeatureURLOptions::All)) == 1, "Update this function to handle the newly added EGameFeatureInstallBundleProtocolOptions value!");
 
 	return ReturnedString;
 }
 
-bool FInstallBundlePluginProtocolMetaData::FromString(FStringView URLString, FInstallBundlePluginProtocolMetaData& Metadata)
+TValueOrError<FInstallBundlePluginProtocolMetaData, void> FInstallBundlePluginProtocolMetaData::FromString(FStringView URLOptionsString)
 {
-	bool bParseSuccess = true;
+	TArray<FName> InstallBundles;
 
-	TArray<FStringView> URLOptions;
-	UE::String::ParseTokens(URLString, UE::GameFeatures::PluginURLStructureInfo::OptionSeperator, URLOptions);
-
-	Metadata = {};
-
-	if (URLOptions.Num() > 0)
+	bool bParseSuccess = UGameFeaturesSubsystem::ParsePluginURLOptions(URLOptionsString, EGameFeatureURLOptions::Bundles,
+	[&InstallBundles](EGameFeatureURLOptions Option, FStringView OptionName, FStringView OptionValue)
 	{
-		ensureAlwaysMsgf(URLOptions[0].StartsWith(UE::GameFeatures::GameFeaturePluginProtocolPrefix(EGameFeaturePluginProtocol::InstallBundle)), TEXT("Unexpected URL Format! Expected Protocol and uplugin information at the beginning of the URL %.*s"), URLString.Len(), URLString.GetData());
-
-		//Parse through our URLOptions. Start at ParsingIndex 1 as option 0 should always be the .uplugin path that doesn't contain metadata
-		for (int ParsingIndex = 1; ParsingIndex < URLOptions.Num(); ++ParsingIndex)
-		{
-			const FStringView URLOption = URLOptions[ParsingIndex];
-	
-			TArray<FStringView> OptionStrings;
-			UE::String::ParseTokens(URLOption, UE::GameFeatures::PluginURLStructureInfo::OptionAssignOperator, OptionStrings);
-
-			if (OptionStrings.Num() > 0)
-			{
-				EGameFeatureInstallBundleProtocolOptions OptionEnum = EGameFeatureInstallBundleProtocolOptions::Count;
-				LexFromString(OptionEnum, OptionStrings[0]);
-
-				switch (OptionEnum)
-				{
-					case EGameFeatureInstallBundleProtocolOptions::Bundles:
-					{
-						if (OptionStrings.Num() != 2)
-						{
-							bParseSuccess = false;
-							UE_LOG(LogGameFeatures, Error, TEXT("Error parsing InstallBundle protocol options URL %.*s. No Valid Bundle List Found!"), URLString.Len(), URLString.GetData());
-						}
-						else
-						{
-							TArray<FStringView> BundleNames;
-							UE::String::ParseTokens(OptionStrings[1], UE::GameFeatures::PluginURLStructureInfo::OptionListSeperator, BundleNames);
-
-							Metadata.InstallBundles.Reserve(BundleNames.Num());
-							for (const FStringView BundleNameString : BundleNames)
-							{
-								Metadata.InstallBundles.Emplace(BundleNameString);
-							}
-						}
-
-						break;
-					}
-
-					case EGameFeatureInstallBundleProtocolOptions::Count:
-					{
-						// Ignore unknown options as they may be game specfic
-						break;
-					}
-
-					default:
-					{
-						bParseSuccess = false;
-						UE_LOG(LogGameFeatures, Error, TEXT("Error parsing InstallBundle protocol options for URL %.*s. Unknown Option %.*s"), URLString.Len(), URLString.GetData(), URLOption.Len(), URLOption.GetData());
+		check(Option == EGameFeatureURLOptions::Bundles);
+		InstallBundles.Emplace(OptionValue);
+	});
 
-						break;
-					}
-				}
-			}
-		}
-	}
 	//We require to have InstallBundle names for this URL parse to be correct
-	if (Metadata.InstallBundles.Num() == 0)
+	if (!bParseSuccess || InstallBundles.Num() == 0)
 	{
 		bParseSuccess = false;
-		UE_LOG(LogGameFeatures, Error, TEXT("Error parsing InstallBundle protocol options URL %.*s. No Bundle List Found!"), URLString.Len(), URLString.GetData());
+		UE_LOG(LogGameFeatures, Error, TEXT("Error parsing InstallBundle protocol options URL %.*s"), URLOptionsString.Len(), URLOptionsString.GetData());
+		return MakeError();
 	}
 
-	static_assert(static_cast<uint8>(EGameFeatureInstallBundleProtocolOptions::Count) == 1, "Update this function to handle the newly added EGameFeatureInstallBundleProtocolOptions value!");
-
-	return bParseSuccess;
+	return MakeValue<FInstallBundlePluginProtocolMetaData>(MoveTemp(InstallBundles));
 }
 
 bool FGameFeaturePluginStateMachineProperties::ParseURL()
@@ -3613,7 +3473,8 @@ bool FGameFeaturePluginStateMachineProperties::ParseURL()
 	}
 
 	FStringView PluginPathFromURL;
-	if (!UGameFeaturesSubsystem::ParsePluginURL(PluginIdentifier.GetFullPluginURL(), nullptr, &PluginPathFromURL))
+	FStringView URLOptions;
+	if (!UGameFeaturesSubsystem::ParsePluginURL(PluginIdentifier.GetFullPluginURL(), nullptr, &PluginPathFromURL, &URLOptions))
 	{
 		return false;
 	}
@@ -3630,13 +3491,16 @@ bool FGameFeaturePluginStateMachineProperties::ParseURL()
 	//Do additional parsing of our Metadata from the options on our remaining URL
 	if (GetPluginProtocol() == EGameFeaturePluginProtocol::InstallBundle)
 	{
-		FInstallBundlePluginProtocolMetaData& MetaData = *ProtocolMetadata.SetSubtype<FInstallBundlePluginProtocolMetaData>();
-		if (!FInstallBundlePluginProtocolMetaData::FromString(PluginIdentifier.GetFullPluginURL(), MetaData))
+		TValueOrError<FInstallBundlePluginProtocolMetaData, void> MaybeMetaData = FInstallBundlePluginProtocolMetaData::FromString(URLOptions);
+		if (MaybeMetaData.HasError())
 		{
 			ensureMsgf(false, TEXT("Failure to parse URL %s into a valid FInstallBundlePluginProtocolMetaData"), *PluginIdentifier.GetFullPluginURL());
 			return false;
 		}
 
+		FInstallBundlePluginProtocolMetaData& MetaData = *ProtocolMetadata.SetSubtype<FInstallBundlePluginProtocolMetaData>();
+		MetaData = MaybeMetaData.StealValue();
+
 		// Add default protocol options if they are not set yet
 		if (!ProtocolOptions.HasSubtype<FInstallBundlePluginProtocolOptions>())
 		{
diff --git a/Engine/Plugins/Runtime/GameFeatures/Source/GameFeatures/Private/GameFeaturePluginStateMachine.h b/Engine/Plugins/Runtime/GameFeatures/Source/GameFeatures/Private/GameFeaturePluginStateMachine.h
--- a/Engine/Plugins/Runtime/GameFeatures/Source/GameFeatures/Private/GameFeaturePluginStateMachine.h
+++ b/Engine/Plugins/Runtime/GameFeatures/Source/GameFeatures/Private/GameFeaturePluginStateMachine.h
@@ -159,11 +159,6 @@ Transition states are expected to transition the machine to another state after
                              +--------------------+
 */
 
-namespace UE::GameFeatures
-{
-	const TCHAR* GameFeaturePluginProtocolPrefix(EGameFeaturePluginProtocol Protocol);
-}
-
 struct FGameFeaturePluginStateRange
 {
 	EGameFeaturePluginState MinState = EGameFeaturePluginState::Uninitialized;
@@ -230,11 +225,14 @@ inline bool operator>(const FGameFeaturePluginStateRange& StateRange, EGameFeatu
 
 struct FInstallBundlePluginProtocolMetaData
 {
+	FInstallBundlePluginProtocolMetaData() = default;
+	FInstallBundlePluginProtocolMetaData(TArray<FName> InInstallBundles) : InstallBundles(MoveTemp(InInstallBundles)) {}
+
 	TArray<FName> InstallBundles;
 
 	/** Functions to convert to/from the URL FString representation of this metadata **/
 	FString ToString() const;
-	static bool FromString(FStringView URLString, FInstallBundlePluginProtocolMetaData& OutMetadata);
+	static TValueOrError<FInstallBundlePluginProtocolMetaData, void> FromString(FStringView URLOptionsString);
 };
 
 struct FGameFeatureProtocolMetadata : public TUnion<FInstallBundlePluginProtocolMetaData, FNull>
diff --git a/Engine/Plugins/Runtime/GameFeatures/Source/GameFeatures/Private/GameFeatureTypes.cpp b/Engine/Plugins/Runtime/GameFeatures/Source/GameFeatures/Private/GameFeatureTypes.cpp
new file mode 100644
--- /dev/null
+++ b/Engine/Plugins/Runtime/GameFeatures/Source/GameFeatures/Private/GameFeatureTypes.cpp
@@ -0,0 +1,50 @@
+// Copyright Epic Games, Inc. All Rights Reserved.
+
+#include "GameFeatureTypes.h"
+
+#include "Containers/UnrealString.h"
+#include "Containers/StringView.h"
+
+namespace UE::GameFeatures
+{
+#define GAME_FEATURE_PLUGIN_STATE_TO_STRING(inEnum, inText) case EGameFeaturePluginState::inEnum: return TEXT(#inEnum);
+	FString ToString(EGameFeaturePluginState InType)
+	{
+		switch (InType)
+		{
+			GAME_FEATURE_PLUGIN_STATE_LIST(GAME_FEATURE_PLUGIN_STATE_TO_STRING)
+		default:
+			check(0);
+			return FString();
+		}
+	}
+#undef GAME_FEATURE_PLUGIN_STATE_TO_STRING
+}
+
+const TCHAR* LexToString(EGameFeatureURLOptions InOption)
+{
+#define GAME_FEATURE_PLUGIN_URL_OPTIONS_STRING(inEnum, inVal) case EGameFeatureURLOptions::inEnum: return TEXT(#inEnum);
+	switch (InOption)
+	{
+		GAME_FEATURE_PLUGIN_URL_OPTIONS_LIST(GAME_FEATURE_PLUGIN_URL_OPTIONS_STRING)
+	}
+#undef GAME_FEATURE_PLUGIN_URL_OPTIONS_STRING
+
+	check(0);
+	return TEXT("");
+}
+
+void LexFromString(EGameFeatureURLOptions& ValueOut, const FStringView& StringIn)
+{
+	ValueOut = EGameFeatureURLOptions::None;
+	
+	for (EGameFeatureURLOptions OptionToCheck : MakeFlagsRange(EGameFeatureURLOptions::All))
+	{
+		if (FStringView(LexToString(OptionToCheck)) == StringIn)
+		{
+			ValueOut = OptionToCheck;
+			return;
+		}
+	}
+}
+
diff --git a/Engine/Plugins/Runtime/GameFeatures/Source/GameFeatures/Private/GameFeaturesSubsystem.cpp b/Engine/Plugins/Runtime/GameFeatures/Source/GameFeatures/Private/GameFeaturesSubsystem.cpp
--- a/Engine/Plugins/Runtime/GameFeatures/Source/GameFeatures/Private/GameFeaturesSubsystem.cpp
+++ b/Engine/Plugins/Runtime/GameFeatures/Source/GameFeatures/Private/GameFeaturesSubsystem.cpp
@@ -21,6 +21,7 @@
 #include "Misc/FileHelper.h"
 #include "Serialization/JsonSerializer.h"
 #include "Stats/StatsMisc.h"
+#include "String/ParseTokens.h"
 #include "Engine/AssetManager.h"
 #include "Engine/StreamableManager.h"
 #include "Engine/AssetManagerSettings.h"
@@ -53,6 +54,19 @@ namespace UE::GameFeatures
 		const TCHAR* CancelAddonCode = TEXT("_Cancel");
 	}
 
+#define GAME_FEATURE_PLUGIN_PROTOCOL_PREFIX(inEnum, inString) case EGameFeaturePluginProtocol::inEnum: return inString;
+	static const TCHAR* GameFeaturePluginProtocolPrefix(EGameFeaturePluginProtocol Protocol)
+	{
+		switch (Protocol)
+		{
+			GAME_FEATURE_PLUGIN_PROTOCOL_LIST(GAME_FEATURE_PLUGIN_PROTOCOL_PREFIX)
+		}
+
+		check(false);
+		return nullptr;
+	}
+#undef GAME_FEATURE_PLUGIN_PROTOCOL_PREFIX
+
 	static bool GCachePluginDetails = true;
 	static FAutoConsoleVariableRef CVarCachePluginDetails(
 		TEXT("GameFeaturePlugin.CachePluginDetails"),
@@ -84,30 +98,6 @@ namespace UE::GameFeatures
 #endif // !UE_BUILD_SHIPPING
 }
 
-#define GAME_FEATURE_PLUGIN_STATE_LEX_TO_STRING(inEnum, inText)  \
-    case(EGameFeaturePluginState::inEnum):                       \
-    {                                                            \
-        return TEXT(#inEnum);                                    \
-    }                                                            
-
-namespace GameFeaturePluginStatePrivate
-{
-	FString LexToString(EGameFeaturePluginState InEnum)
-	{
-		switch (InEnum)
-		{
-			GAME_FEATURE_PLUGIN_STATE_LIST(GAME_FEATURE_PLUGIN_STATE_LEX_TO_STRING)
-
-		default:
-			{
-				ensureAlwaysMsgf(false, TEXT("Logic error causing a missing LexToString value for EGameFeaturePluginState:%d"), static_cast<uint8>(InEnum));
-				return TEXT("ERROR_UNSUPPORTED_ENUM");
-			}
-		}
-	}
-}
-#undef GAME_FEATURE_PLUGIN_STATE_LEX_TO_STRING
-
 const FString LexToString(const EBuiltInAutoState BuiltInAutoState)
 {
 	switch (BuiltInAutoState)
@@ -710,7 +700,6 @@ FString GetPluginURL_InstallBundleProtocol(const FString& PluginName, const FIns
 	FString Path;
 	Path += UE::GameFeatures::GameFeaturePluginProtocolPrefix(EGameFeaturePluginProtocol::InstallBundle);
 	Path += PluginName;
-	Path += UE::GameFeatures::PluginURLStructureInfo::OptionSeperator;
 	Path += ProtocolMetadata.ToString();
 	if (AdditionalOptions.Num() > 0)
 	{
@@ -823,6 +812,127 @@ bool UGameFeaturesSubsystem::ParsePluginURL(FStringView PluginURL, EGameFeatureP
 	return false;
 }
 
+namespace GameFeaturesSubsystem
+{
+	static bool SplitOption(FStringView OptionPair, FStringView& OutOptionName, FStringView& OutOptionValue)
+	{
+		int32 TokenCount = 0;
+		FStringView OptionName;
+		FStringView OptionValue;
+		UE::String::ParseTokens(OptionPair, UE::GameFeatures::PluginURLStructureInfo::OptionAssignOperator, 
+		[&TokenCount, &OptionName, &OptionValue](FStringView Token)
+		{
+			++TokenCount;
+			switch (TokenCount)
+			{
+			case 1:
+				OptionName = Token;
+				break;
+			case 2:
+				OptionValue = Token;
+				break;
+			}
+		});
+
+		const bool bSuccess = TokenCount == 2;
+		if (bSuccess)
+		{
+			OutOptionName = OptionName;
+			OutOptionValue = OptionValue;
+		}
+
+		return bSuccess;
+	}
+
+	struct FParsePluginURLOptionsFilter
+	{
+		EGameFeatureURLOptions OptionsFlags;
+		TConstArrayView<FStringView> AdditionalOptions;
+	};
+
+	static bool ParsePluginURLOptions(FStringView URLOptionsString, const FParsePluginURLOptionsFilter* OptionsFilter,
+		TFunctionRef<void(EGameFeatureURLOptions Option, FStringView OptionString, FStringView OptionValue)> Output)
+	{
+		enum class EParseState : uint8
+		{
+			First,
+			OK,
+			Error
+		};
+
+		//Parse through our URLOptions. The first option won't appear until after the first seperator.
+		//We don't care what comes before the first seperator
+		EParseState ParseState = EParseState::First;
+		UE::String::ParseTokens(URLOptionsString, UE::GameFeatures::PluginURLStructureInfo::OptionSeperator, 
+		[OptionsFilter, Output, &ParseState](FStringView Token)
+		{
+			if (ParseState == EParseState::Error)
+			{
+				return;
+			}
+
+			if (ParseState == EParseState::First)
+			{
+				ParseState = EParseState::OK;
+				return;
+			}
+
+			FStringView OptionName;
+			FStringView OptionValue;
+			if (!GameFeaturesSubsystem::SplitOption(Token, OptionName, OptionValue))
+			{
+				ParseState = EParseState::Error;
+				return;
+			}
+
+			EGameFeatureURLOptions OptionEnum = EGameFeatureURLOptions::None;
+
+			if (!OptionsFilter || OptionsFilter->OptionsFlags != EGameFeatureURLOptions::None)
+			{
+				LexFromString(OptionEnum, OptionName);
+			}
+
+			if (!OptionsFilter || EnumHasAnyFlags(OptionsFilter->OptionsFlags, OptionEnum) || OptionsFilter->AdditionalOptions.Contains(OptionName))
+			{
+				UE::String::ParseTokens(OptionValue, UE::GameFeatures::PluginURLStructureInfo::OptionListSeperator,
+				[Output, OptionEnum, OptionName](FStringView ListToken)
+				{
+					Output(OptionEnum, OptionName, ListToken);
+				});
+			}
+		});
+
+		return ParseState == EParseState::OK;
+	}
+}
+
+bool UGameFeaturesSubsystem::ParsePluginURLOptions(FStringView URLOptionsString,
+	TFunctionRef<void(EGameFeatureURLOptions Option, FStringView OptionString, FStringView OptionValue)> Output)
+{
+	return GameFeaturesSubsystem::ParsePluginURLOptions(URLOptionsString, nullptr, Output);
+}
+
+bool UGameFeaturesSubsystem::ParsePluginURLOptions(FStringView URLOptionsString, EGameFeatureURLOptions OptionsFlags,
+	TFunctionRef<void(EGameFeatureURLOptions Option, FStringView OptionString, FStringView OptionValue)> Output)
+{
+	const GameFeaturesSubsystem::FParsePluginURLOptionsFilter OptionsFilter{ OptionsFlags, {} };
+	return GameFeaturesSubsystem::ParsePluginURLOptions(URLOptionsString, &OptionsFilter, Output);
+}
+
+bool UGameFeaturesSubsystem::ParsePluginURLOptions(FStringView URLOptionsString, TConstArrayView<FStringView> AdditionalOptions,
+	TFunctionRef<void(EGameFeatureURLOptions Option, FStringView OptionString, FStringView OptionValue)> Output)
+{
+	const GameFeaturesSubsystem::FParsePluginURLOptionsFilter OptionsFilter{ EGameFeatureURLOptions::None, AdditionalOptions };
+	return GameFeaturesSubsystem::ParsePluginURLOptions(URLOptionsString, &OptionsFilter, Output);
+}
+
+bool UGameFeaturesSubsystem::ParsePluginURLOptions(FStringView URLOptionsString, EGameFeatureURLOptions OptionsFlags, TConstArrayView<FStringView> AdditionalOptions,
+	TFunctionRef<void(EGameFeatureURLOptions Option, FStringView OptionString, FStringView OptionValue)> Output)
+{
+	const GameFeaturesSubsystem::FParsePluginURLOptionsFilter OptionsFilter{ OptionsFlags, AdditionalOptions };
+	return GameFeaturesSubsystem::ParsePluginURLOptions(URLOptionsString, &OptionsFilter, Output);
+}
+
 void UGameFeaturesSubsystem::OnGameFeatureTerminating(const FString& PluginName, const FGameFeaturePluginIdentifier& PluginIdentifier)
 {
 	CallbackObservers(EObserverCallback::Terminating, PluginIdentifier, &PluginName);
@@ -2174,9 +2284,8 @@ struct FGameFeaturePluginPredownloadContext : public FGameFeaturePluginPredownlo
 				continue;
 			}
 
-			FInstallBundlePluginProtocolMetaData InstallBundleOptions;
-			const bool bParsedBundles = FInstallBundlePluginProtocolMetaData::FromString(URL, InstallBundleOptions);
-			if (!bParsedBundles)
+			TValueOrError<FInstallBundlePluginProtocolMetaData, void> MaybeInstallBundleOptions = FInstallBundlePluginProtocolMetaData::FromString(URL);
+			if (MaybeInstallBundleOptions.HasError())
 			{
 				UE_LOGFMT(LogGameFeatures, Error, "GFP Predownload failed to parse URL {URL}", ("URL", URL));
 				UE::GameFeatures::FResult ErrorResult = MakeError(FString::Printf(TEXT("%.*s%s"),
@@ -2186,7 +2295,7 @@ struct FGameFeaturePluginPredownloadContext : public FGameFeaturePluginPredownlo
 				return;
 			}
 
-			GFPs.Emplace(URL, MoveTemp(InstallBundleOptions));
+			GFPs.Emplace(URL, MaybeInstallBundleOptions.StealValue());
 		}
 
 		if (GFPs.Num() == 0)
@@ -2984,7 +3093,7 @@ TSet<FString> UGameFeaturesSubsystem::GetActivePluginNames() const
 
 namespace GameFeaturesSubsystem
 { 
-	bool IsContentWithinActivePlugin(const FString& InObjectOrPackagePath, const TSet<FString>& ActivePluginNames)
+	static bool IsContentWithinActivePlugin(const FString& InObjectOrPackagePath, const TSet<FString>& ActivePluginNames)
 	{
 		// Look for the first slash beyond the first one we start with.
 		const int32 RootEndIndex = InObjectOrPackagePath.Find(TEXT("/"), ESearchCase::IgnoreCase, ESearchDir::FromStart, 1);
diff --git a/Engine/Plugins/Runtime/GameFeatures/Source/GameFeatures/Public/GameFeatureTypes.h b/Engine/Plugins/Runtime/GameFeatures/Source/GameFeatures/Public/GameFeatureTypes.h
--- a/Engine/Plugins/Runtime/GameFeatures/Source/GameFeatures/Public/GameFeatureTypes.h
+++ b/Engine/Plugins/Runtime/GameFeatures/Source/GameFeatures/Public/GameFeatureTypes.h
@@ -2,10 +2,13 @@
 
 #pragma once
 
+#include "Containers/StringFwd.h"
+#include "Misc/EnumClassFlags.h"
 #include "Misc/EnumRange.h"
 
 class FString;
 
+/** GFP States */
 #define GAME_FEATURE_PLUGIN_STATE_LIST(XSTATE)	\
 	XSTATE(Uninitialized,					NSLOCTEXT("GameFeatures", "UninitializedStateDisplayName", "Uninitialized"))									/* Unset. Not yet been set up. */ \
 	XSTATE(Terminal,						NSLOCTEXT("GameFeatures", "TerminalStateDisplayName", "Terminal"))												/* Final State before removal of the state machine. */ \
@@ -50,13 +53,16 @@ namespace GameFeaturePluginStatePrivate
 		GAME_FEATURE_PLUGIN_STATE_LIST(GAME_FEATURE_PLUGIN_STATE_ENUM)
 		MAX
 	};
-
-	//defined in GameFeaturesSubsystem.cpp
-	FString LexToString(EGameFeaturePluginState InEnum);
 }
-typedef GameFeaturePluginStatePrivate::EGameFeaturePluginState EGameFeaturePluginState;
+using EGameFeaturePluginState = GameFeaturePluginStatePrivate::EGameFeaturePluginState;
 #undef GAME_FEATURE_PLUGIN_STATE_ENUM
 
+namespace UE::GameFeatures
+{
+	GAMEFEATURES_API FString ToString(EGameFeaturePluginState InType);
+}
+
+/** GFP Protocols */
 #define GAME_FEATURE_PLUGIN_PROTOCOL_LIST(XPROTO)	\
 	XPROTO(File,			TEXT("file:"))			\
 	XPROTO(InstallBundle,	TEXT("installbundle:"))	\
@@ -72,6 +78,26 @@ enum class EGameFeaturePluginProtocol : uint8
 
 ENUM_RANGE_BY_COUNT(EGameFeaturePluginProtocol, EGameFeaturePluginProtocol::Count);
 
+/** GFP URL Options */
+#define GAME_FEATURE_PLUGIN_URL_OPTIONS_LIST(XOPTION)	\
+	XOPTION(Bundles,	(1 << 0))
+
+#define GAME_FEATURE_PLUGIN_URL_OPTIONS_ENUM(inEnum, inVal) inEnum = inVal,
+#define GAME_FEATURE_PLUGIN_URL_OPTIONS_ALL(inEnum, inVal) | inEnum
+enum class EGameFeatureURLOptions : uint8
+{
+	None,
+	GAME_FEATURE_PLUGIN_URL_OPTIONS_LIST(GAME_FEATURE_PLUGIN_URL_OPTIONS_ENUM)
+	All = None GAME_FEATURE_PLUGIN_URL_OPTIONS_LIST(GAME_FEATURE_PLUGIN_URL_OPTIONS_ALL)
+};
+#undef GAME_FEATURE_PLUGIN_URL_OPTIONS_ALL
+#undef GAME_FEATURE_PLUGIN_URL_OPTIONS_ENUM
+
+ENUM_CLASS_FLAGS(EGameFeatureURLOptions);
+
+const TCHAR* LexToString(EGameFeatureURLOptions);
+void LexFromString(EGameFeatureURLOptions& ValueOut, const FStringView& StringIn);
+
 #if UE_ENABLE_INCLUDE_ORDER_DEPRECATED_IN_5_2
 #include "GameFeatureTypesFwd.h"
 #endif
diff --git a/Engine/Plugins/Runtime/GameFeatures/Source/GameFeatures/Public/GameFeatureTypesFwd.h b/Engine/Plugins/Runtime/GameFeatures/Source/GameFeatures/Public/GameFeatureTypesFwd.h
--- a/Engine/Plugins/Runtime/GameFeatures/Source/GameFeatures/Public/GameFeatureTypesFwd.h
+++ b/Engine/Plugins/Runtime/GameFeatures/Source/GameFeatures/Public/GameFeatureTypesFwd.h
@@ -17,3 +17,5 @@ namespace UE::GameFeatures
 {
 	GAMEFEATURES_API FString ToString(EGameFeaturePluginState InType);
 }
+
+enum class EGameFeatureURLOptions : uint8;
diff --git a/Engine/Plugins/Runtime/GameFeatures/Source/GameFeatures/Public/GameFeaturesSubsystem.h b/Engine/Plugins/Runtime/GameFeatures/Source/GameFeatures/Public/GameFeaturesSubsystem.h
--- a/Engine/Plugins/Runtime/GameFeatures/Source/GameFeatures/Public/GameFeaturesSubsystem.h
+++ b/Engine/Plugins/Runtime/GameFeatures/Source/GameFeatures/Public/GameFeaturesSubsystem.h
@@ -480,6 +480,17 @@ public:
 	/** Parse the plugin URL into subparts */
 	static bool ParsePluginURL(FStringView PluginURL, EGameFeaturePluginProtocol* OutProtocol = nullptr, FStringView* OutPath = nullptr, FStringView* OutOptions = nullptr);
 
+	/** Parse options from a plugin URL or the options subpart of the plugin URL */
+	static bool ParsePluginURLOptions(FStringView URLOptionsString,
+		TFunctionRef<void(EGameFeatureURLOptions Option, FStringView OptionString, FStringView OptionValue)> Output);
+	static bool ParsePluginURLOptions(FStringView URLOptionsString, EGameFeatureURLOptions OptionsFlags,
+		TFunctionRef<void(EGameFeatureURLOptions Option, FStringView OptionString, FStringView OptionValue)> Output);
+	static bool ParsePluginURLOptions(FStringView URLOptionsString, TConstArrayView<FStringView> AdditionalOptions,
+		TFunctionRef<void(EGameFeatureURLOptions Option, FStringView OptionString, FStringView OptionValue)> Output);
+	static bool ParsePluginURLOptions(FStringView URLOptionsString, EGameFeatureURLOptions OptionsFlags, TConstArrayView<FStringView> AdditionalOptions,
+		TFunctionRef<void(EGameFeatureURLOptions Option, FStringView OptionString, FStringView OptionValue)> Output);
+
+
 public:
 	/** Returns all the active plugins GameFeatureDatas */
 	void GetGameFeatureDataForActivePlugins(TArray<const UGameFeatureData*>& OutActivePluginFeatureDatas);
