CIS fix - only the generic version of CountBits is constexpr

diff --git a/Engine/Plugins/Runtime/GameFeatures/Source/GameFeatures/Private/GameFeaturePluginStateMachine.cpp b/Engine/Plugins/Runtime/GameFeatures/Source/GameFeatures/Private/GameFeaturePluginStateMachine.cpp
--- a/Engine/Plugins/Runtime/GameFeatures/Source/GameFeatures/Private/GameFeaturePluginStateMachine.cpp
+++ b/Engine/Plugins/Runtime/GameFeatures/Source/GameFeatures/Private/GameFeaturePluginStateMachine.cpp
@@ -3438,7 +3438,8 @@ FString FInstallBundlePluginProtocolMetaData::ToString() const
 	});
 	ReturnedString.Append(BundlesList);
 
-	static_assert(FMath::CountBits(static_cast<uint64>(EGameFeatureURLOptions::All)) == 1, "Update this function to handle the newly added EGameFeatureInstallBundleProtocolOptions value!");
+	// Only the generic version of CountBits is constexpr...
+	static_assert(FGenericPlatformMath::CountBits(static_cast<uint64>(EGameFeatureURLOptions::All)) == 1, "Update this function to handle the newly added EGameFeatureInstallBundleProtocolOptions value!");
 
 	return ReturnedString;
 }
