Content Bundle: always inject content when running a commandlet, not just during cooks.
This will make sure to generate HLODs for Content Bundles that are registered.

diff --git a/Engine/Plugins/Runtime/GameFeatures/Source/GameFeatures/Private/GameFeatureAction_AddWPContent.cpp b/Engine/Plugins/Runtime/GameFeatures/Source/GameFeatures/Private/GameFeatureAction_AddWPContent.cpp
--- a/Engine/Plugins/Runtime/GameFeatures/Source/GameFeatures/Private/GameFeatureAction_AddWPContent.cpp
+++ b/Engine/Plugins/Runtime/GameFeatures/Source/GameFeatures/Private/GameFeatureAction_AddWPContent.cpp
@@ -31,7 +31,7 @@ void UGameFeatureAction_AddWPContent::OnGameFeatureRegistering()
 	ContentBundleClient = FContentBundleClient::CreateClient(ContentBundleDescriptor, GetTypedOuter<UGameFeatureData>()->GetName());
 	
 #if WITH_EDITOR
-	if (IsRunningCookCommandlet() && ContentBundleClient != nullptr)
+	if (IsRunningCommandlet() && ContentBundleClient != nullptr)
 	{
 		ContentBundleClient->RequestContentInjection();
 	}
