Remove uninitialized and unreferenced member.

diff --git a/Engine/Plugins/Runtime/GameFeatures/Source/GameFeatures/Public/GameFeaturesSubsystem.h b/Engine/Plugins/Runtime/GameFeatures/Source/GameFeatures/Public/GameFeaturesSubsystem.h
--- a/Engine/Plugins/Runtime/GameFeatures/Source/GameFeatures/Public/GameFeaturesSubsystem.h
+++ b/Engine/Plugins/Runtime/GameFeatures/Source/GameFeatures/Public/GameFeaturesSubsystem.h
@@ -394,7 +394,6 @@ struct FGameFeatureInfo
 	FString URL;
 	bool bLoadedAsBuiltIn;
 	EGameFeaturePluginState CurrentState;
-	EGameFeaturePluginState DestinationState;
 };
 
 /** The manager subsystem for game features */
