Infrastructure and general implementation of a new system for providing improve error messages when we skip loading a package during import because the plugin wasn't mounted. the plugin manager and the GameFeatureSubsystem now have an opportunity to provide additional information to the user about why the plugin wasn't loaded or the package wasn't found.

diff --git a/Engine/Plugins/Runtime/GameFeatures/Source/GameFeatures/Private/GameFeaturesSubsystem.cpp b/Engine/Plugins/Runtime/GameFeatures/Source/GameFeatures/Private/GameFeaturesSubsystem.cpp
--- a/Engine/Plugins/Runtime/GameFeatures/Source/GameFeatures/Private/GameFeaturesSubsystem.cpp
+++ b/Engine/Plugins/Runtime/GameFeatures/Source/GameFeatures/Private/GameFeaturesSubsystem.cpp
@@ -491,11 +491,15 @@ void UGameFeaturesSubsystem::Initialize(FSubsystemCollectionBase& Collection)
 				}
 			}),
 		ECVF_Cheat);
+
+	GetExplanationForUnavailablePackageDelegateHandle = IPluginManager::Get().GetExplanationForUnavailablePackageWithPluginInfoDelegate().AddUObject(this, &UGameFeaturesSubsystem::GetExplanationForUnavailablePackage);
 }
 
 void UGameFeaturesSubsystem::Deinitialize()
 {
 	UE_LOG(LogGameFeatures, Log, TEXT("Shutting down game features subsystem"));
+	IPluginManager::Get().GetExplanationForUnavailablePackageWithPluginInfoDelegate().Remove(GetExplanationForUnavailablePackageDelegateHandle);
+	GetExplanationForUnavailablePackageDelegateHandle.Reset();
 
 	if ((GameSpecificPolicies != nullptr) && bInitializedPolicyManager)
 	{
@@ -694,6 +698,27 @@ FString UGameFeaturesSubsystem::GetPluginURL_FileProtocol(const FString& PluginD
 	return Path;
 }
 
+void UGameFeaturesSubsystem::GetExplanationForUnavailablePackage(const FString& UnavailablePackage, IPlugin* PluginIfFound, FStringBuilderBase& InOutExplanation)
+{
+#if WITH_EDITOR
+	if (PluginIfFound)
+	{
+		if (FString* Explanation = UnmountedPluginNameToExplanation.Find(PluginIfFound->GetName()))
+		{
+			InOutExplanation.Appendf(TEXT("\nUGameFeaturesSubsystem: Explanation for not mounting plugin %s: %s"), *PluginIfFound->GetFriendlyName(), **Explanation);
+		}
+	}
+	else
+	{
+		FString ContentDirName = FPackageName::SplitPackageNameRoot(UnavailablePackage, nullptr).GetData();
+		if (FString* Explanation = UnmountedPluginNameToExplanation.Find(ContentDirName))
+		{
+			InOutExplanation.Appendf(TEXT("\nUGameFeaturesSubsystem: Explanation for not mounting plugin %s: %s"), *ContentDirName, **Explanation);
+		}
+	}
+#endif
+}
+
 FString GetPluginURL_InstallBundleProtocol(const FString& PluginName, const FInstallBundlePluginProtocolMetaData& ProtocolMetadata, TArrayView<const TPair<FString, FString>> AdditionalOptions = {})
 {
 	ensure(ProtocolMetadata.InstallBundles.Num() > 0);
@@ -2596,6 +2621,15 @@ void UGameFeaturesSubsystem::LoadBuiltInGameFeaturePluginComplete(const UE::Game
 			*UE::GameFeatures::ToString(Machine->GetCurrentState()), 
 			*UE::GameFeatures::ToString(RequestedDestination.MinState), 
 			*UE::GameFeatures::ToString(RequestedDestination.MaxState));
+#if WITH_EDITOR
+		if (Machine->GetCurrentState() > EGameFeaturePluginState::Mounting)
+		{
+			if (FString* Explanation = UnmountedPluginNameToExplanation.Find(Machine->GetPluginName()))
+			{
+				UnmountedPluginNameToExplanation.Remove(Machine->GetPluginName());
+			}
+		}
+#endif
 	}
 	else
 	{
@@ -2606,6 +2640,7 @@ void UGameFeaturesSubsystem::LoadBuiltInGameFeaturePluginComplete(const UE::Game
 			*UE::GameFeatures::ToString(Machine->GetDestination().MinState),
 			*UE::GameFeatures::ToString(Machine->GetDestination().MaxState),
 			*ErrorMessage);
+		SetExplanationForNotMountingPlugin(Machine->GetPluginURL(), ErrorMessage);
 	}
 }
 
@@ -2670,6 +2705,15 @@ void UGameFeaturesSubsystem::ChangeGameFeatureDestination(UGameFeaturePluginStat
 
 void UGameFeaturesSubsystem::ChangeGameFeatureTargetStateComplete(UGameFeaturePluginStateMachine* Machine, const UE::GameFeatures::FResult& Result, FGameFeaturePluginChangeStateComplete CompleteDelegate)
 {
+#if WITH_EDITOR
+	if (!Result.HasError() && Machine->GetCurrentState() > EGameFeaturePluginState::Mounting)
+	{
+		if (FString* Explanation = UnmountedPluginNameToExplanation.Find(Machine->GetPluginName()))
+		{
+			UnmountedPluginNameToExplanation.Remove(Machine->GetPluginName());
+		}
+	}
+#endif
 	CompleteDelegate.ExecuteIfBound(Result);
 }
 
@@ -2708,7 +2752,10 @@ bool UGameFeaturesSubsystem::FindOrCreatePluginDependencyStateMachines(const FSt
 			TValueOrError<FString, FString> DependencyURLInfo = GameSpecificPolicies->ResolvePluginDependency(PluginURL, DependencyName);
 			if (DependencyURLInfo.HasError())
 			{
-				UE_LOG(LogGameFeatures, Error, TEXT("Game feature plugin '%s' has unknown dependency '%s' [%s]."), *PluginFilename, *DependencyName, *DependencyURLInfo.GetError());
+				FString ErrorMessage = FString::Printf(TEXT("Game feature plugin '%s' has unknown dependency '%s' [%s]."), *PluginFilename, *DependencyName, *DependencyURLInfo.GetError());
+				UE_LOG(LogGameFeatures, Error, TEXT("%s"), *ErrorMessage);
+
+				SetExplanationForNotMountingPlugin(PluginURL, ErrorMessage);
 
 				//Don't actually return false here as we want to still be able to progress in the case of 
 				//things like an editor plugin being included as a dependency in the client or a dynamic dependency that
@@ -3241,3 +3288,16 @@ void UGameFeaturesSubsystem::GetPluginsToCook(TSet<FString>& OutPlugins)
 	
 	OutPlugins.Append(PluginsList);
 }
+
+void UGameFeaturesSubsystem::SetExplanationForNotMountingPlugin(const FString& PluginURL, const FString& Explanation)
+{
+#if WITH_EDITOR
+	FGameFeaturePluginIdentifier Identifier(PluginURL);
+	FStringView PluginName = Identifier.GetPluginName();
+	if (!PluginName.IsEmpty())
+	{
+		UnmountedPluginNameToExplanation.FindOrAdd(FString(PluginName)) = Explanation;
+	}
+#endif
+}
+
diff --git a/Engine/Plugins/Runtime/GameFeatures/Source/GameFeatures/Public/GameFeaturesSubsystem.h b/Engine/Plugins/Runtime/GameFeatures/Source/GameFeatures/Public/GameFeaturesSubsystem.h
--- a/Engine/Plugins/Runtime/GameFeatures/Source/GameFeatures/Public/GameFeaturesSubsystem.h
+++ b/Engine/Plugins/Runtime/GameFeatures/Source/GameFeatures/Public/GameFeaturesSubsystem.h
@@ -693,6 +693,9 @@ private:
 
 	static void RemoveGameFeatureFromAssetManager(const UGameFeatureData* GameFeatureToRemove, const FString& PluginName, const TArray<FName>& AddedPrimaryAssetTypes);
 
+	// Provide additional causal information when a package is unavailable for load
+	void GetExplanationForUnavailablePackage(const FString& SkippedPackage, IPlugin* PluginIfFound, FStringBuilderBase& InOutExplanation);
+
 private:
 	bool ShouldUpdatePluginProtocolOptions(const UGameFeaturePluginStateMachine* StateMachine, const FGameFeatureProtocolOptions& NewOptions);
 	UE::GameFeatures::FResult UpdateGameFeatureProtocolOptions(UGameFeaturePluginStateMachine* StateMachine, const FGameFeatureProtocolOptions& NewOptions, bool* bOutDidUpdate = nullptr);
@@ -753,6 +756,8 @@ private:
 	/** Handle 'ListGameFeaturePlugins' console command */
 	void ListGameFeaturePlugins(const TArray<FString>& Args, UWorld* InWorld, FOutputDevice& Ar);
 
+	void SetExplanationForNotMountingPlugin(const FString& PluginURL, const FString& Explanation);
+
 	enum class EObserverCallback
 	{
 		CheckingStatus,
@@ -803,6 +808,12 @@ private:
 	UPROPERTY(Transient)
 	TObjectPtr<UGameFeaturesProjectPolicies> GameSpecificPolicies;
 
+#if WITH_EDITOR
+	// When we decide not to mount a plugin, we can store an explanation here so that if we later attempt to load an asset from it we can tell the user why it's not available
+	TMap<FString, FString> UnmountedPluginNameToExplanation;
+#endif
+	FDelegateHandle GetExplanationForUnavailablePackageDelegateHandle;
+
 	bool bInitializedPolicyManager = false;
 };
 
diff --git a/Engine/Source/Runtime/CoreUObject/Private/Misc/PackageName.cpp b/Engine/Source/Runtime/CoreUObject/Private/Misc/PackageName.cpp
--- a/Engine/Source/Runtime/CoreUObject/Private/Misc/PackageName.cpp
+++ b/Engine/Source/Runtime/CoreUObject/Private/Misc/PackageName.cpp
@@ -1696,6 +1696,66 @@ bool FPackageName::TryGetMountPointForPath(FStringView InFilePathOrPackageName,
 	return false;
 }
 
+FPackageName::FGetExplanationForUnavailablePackageDelegate& FPackageName::GetExplanationForUnavailablePackageDelegate()
+{
+	static FPackageName::FGetExplanationForUnavailablePackageDelegate Delegate;
+	return Delegate;
+}
+
+void FPackageName::GetExplanationForUnavailablePackage(const FName& UnavailablePackageName, FStringBuilderBase& InOutExplanation)
+{
+	FString PackageRoot = FString::Printf(TEXT("/%s"), *FPackageName::SplitPackageNameRoot(UnavailablePackageName, nullptr));
+	if (FPackageName::MountPointExists(PackageRoot))
+	{
+		// The mount point exists but we failed to load the asset. Maybe the package itself didn't exist?
+		const EPackageLocationFilter BothLocationsMask = (EPackageLocationFilter)((uint8)EPackageLocationFilter::FileSystem | (uint8)EPackageLocationFilter::IoDispatcher);
+		FPackagePath SkippedPackagePath;
+		const FString& UnavailablePackageNameString = UnavailablePackageName.ToString();
+		bool ConvertedToPath = FPackagePath::TryFromMountedName(UnavailablePackageNameString, SkippedPackagePath);
+		if (ensure(ConvertedToPath))
+		{
+			EPackageLocationFilter PackageLocations = DoesPackageExistEx(SkippedPackagePath, BothLocationsMask);
+
+			if (PackageLocations == EPackageLocationFilter::None)
+			{
+				InOutExplanation.Appendf(TEXT("FPackageName: Skipped package %s has a valid, mounted, mount point but does not exist either on disk or in iostore."),
+					*UnavailablePackageNameString);
+				const FString& LocalFullPath = SkippedPackagePath.GetLocalFullPath();
+				if (LocalFullPath.Len())
+				{
+					FString PathToPrint = IFileManager::Get().ConvertToAbsolutePathForExternalAppForRead(*LocalFullPath);
+					if (PathToPrint.Len() == 0)
+					{
+						PathToPrint = LocalFullPath;
+					}
+					InOutExplanation.Appendf(TEXT(" The uncooked file would be expected on disk at '%s'. Perhaps it has been deleted or was not synced?\n"), *PathToPrint);
+				}
+				else
+				{
+					InOutExplanation.Append(TEXT("\n"));
+				}
+			}
+			else
+			{
+				InOutExplanation.Appendf(TEXT("FPackageName: Skipped package %s was found %s but for some reason could not be found or used by the loader. This could occur if the loader is only considering cooked packages but the package exists only on disk.\n"),
+					*UnavailablePackageNameString,
+					(PackageLocations == BothLocationsMask) ? TEXT("on disk and in iodispatcher") :
+					((PackageLocations == EPackageLocationFilter::FileSystem) ? TEXT("on disk") : TEXT("in iodispatcher")));
+			}
+		}
+	}
+	else
+	{
+		InOutExplanation.Appendf(TEXT("FPackageName: Unable to identify a valid mount point associated with skipped package %s. The package root is unknown.\n"), *UnavailablePackageName.ToString());
+	}
+
+	// The plugin manager needs to be called directly rather than via delegate due to dependency issue (CoreUObject depends on Projects and not vice versa)
+	IPluginManager::Get().GetExplanationForUnavailablePackage(UnavailablePackageName, InOutExplanation);
+
+	// Give other systems an opportunity to add to this explanation
+	GetExplanationForUnavailablePackageDelegate().Broadcast(UnavailablePackageName, InOutExplanation);
+}
+
 FString FPackageName::GetModuleScriptPackageName(FStringView InModuleName)
 {
 	return FString::Printf(TEXT("/Script/%.*s"), InModuleName.Len(), InModuleName.GetData());
diff --git a/Engine/Source/Runtime/CoreUObject/Private/Serialization/AsyncLoading2.cpp b/Engine/Source/Runtime/CoreUObject/Private/Serialization/AsyncLoading2.cpp
--- a/Engine/Source/Runtime/CoreUObject/Private/Serialization/AsyncLoading2.cpp
+++ b/Engine/Source/Runtime/CoreUObject/Private/Serialization/AsyncLoading2.cpp
@@ -83,6 +83,7 @@
 #include "UObject/CoreRedirects.h"
 #include "Serialization/ZenPackageHeader.h"
 #include "Trace/Trace.h"
+#include "Containers/Ticker.h"
 
 #include <atomic>
 
@@ -5182,6 +5183,31 @@ void FAsyncPackage2::ImportPackagesRecursiveInner(FAsyncLoadingThreadState2& Thr
 			{
 				UE_ASYNC_PACKAGE_CLOG(!ImportedPackageUPackageName.IsNone(), Display, Desc, TEXT("ImportPackages: SkipPackage"),
 					TEXT("Skipping non mounted imported package %s (0x%llX)"), *ImportedPackageNameToLoad.ToString(), ImportedPackageId.ValueForDebugging());
+				if (!ImportedPackageUPackageName.IsNone())
+				{
+					FName PackagePathToLoadName = Desc.PackagePathToLoad.GetPackageFName();
+					ExecuteOnGameThread(TEXT("GetExplanationForUnavailablePackage"), [ImportedPackageNameToLoad, PackagePathToLoadName]()
+					{
+						FString ImportedPackageNameToLoadString = ImportedPackageNameToLoad.ToString();
+						FString PackagePathToLoadString = PackagePathToLoadName.ToString();
+						TStringBuilder<2048> Explanation;
+						FPackageName::GetExplanationForUnavailablePackage(ImportedPackageNameToLoad, Explanation);
+						if (Explanation.Len())
+						{
+							UE_LOG(LogStreaming, Warning, TEXT("While trying to load package %s, a dependent package %s was not available. Additional explanatory information follows:\n%s"), 
+								*PackagePathToLoadString,
+								*ImportedPackageNameToLoadString,
+								Explanation.ToString());
+						}
+						else
+						{
+							UE_LOG(LogStreaming, Warning, TEXT("While trying to load package %s, a dependent package %s was not available. No additional explanation was available."),
+								*PackagePathToLoadString,
+								*ImportedPackageNameToLoadString);
+						}
+						return false;
+					});
+				}
 				ImportedPackageRef.SetIsMissingPackage();
 			}
 			continue;
diff --git a/Engine/Source/Runtime/CoreUObject/Public/Misc/PackageName.h b/Engine/Source/Runtime/CoreUObject/Public/Misc/PackageName.h
--- a/Engine/Source/Runtime/CoreUObject/Public/Misc/PackageName.h
+++ b/Engine/Source/Runtime/CoreUObject/Public/Misc/PackageName.h
@@ -296,6 +296,8 @@ public:
 
 	/**
 	 * This will insert a mount point at the head of the search chain (so it can overlap an existing mount point and win).
+	 * If you register a mount point (even if you do so only in certain circumstances) consider also adding a handler for GetExplanationForUnavailablePackage
+	 * to help debug cases where a package can't be found
 	 *
 	 * @param RootPath Logical Root Path.
 	 * @param ContentPath Content Path on disk.
@@ -855,6 +857,23 @@ public:
 	static COREUOBJECT_API bool TryGetMountPointForPath(FStringView InFilePathOrPackageName, FStringBuilderBase& OutMountPointPackageName, FStringBuilderBase& OutMountPointFilePath, FStringBuilderBase& OutRelPath,
 		EFlexNameType* OutFlexNameType = nullptr, EErrorCode* OutFailureReason = nullptr);
 
+	/** Delegate type for allowing higher levels systems to provide additional context for why a package is not available */
+	DECLARE_MULTICAST_DELEGATE_TwoParams(FGetExplanationForUnavailablePackageDelegate, const FName& /*Unavailable Package*/, FStringBuilderBase& /*InOutExplanation*/);
+
+	/**
+	 * Accessor for installing callbacks for use with GetExplanationForUnavailablePackage
+	 * 
+	 */ 
+	static COREUOBJECT_API FGetExplanationForUnavailablePackageDelegate& GetExplanationForUnavailablePackageDelegate();
+
+	/**
+	 * Attempt to generate an explanation for why a particular package was not available for loading
+	 * 
+	 * @param UnavailablePackageName The package that was requested but could not be loaded
+	 * @param InOutExplanation A string builder that will be populated with any available information about why the package was not available. The string builder will *not* be reset and will be appended to.
+	 */
+	 static COREUOBJECT_API void GetExplanationForUnavailablePackage(const FName& UnavailablePackageName, FStringBuilderBase& InOutExplanation);
+
 private:
 
 	/**
diff --git a/Engine/Source/Runtime/Projects/Private/PluginManager.cpp b/Engine/Source/Runtime/Projects/Private/PluginManager.cpp
--- a/Engine/Source/Runtime/Projects/Private/PluginManager.cpp
+++ b/Engine/Source/Runtime/Projects/Private/PluginManager.cpp
@@ -3539,4 +3539,36 @@ void FPluginManager::RemoveFromModuleNameToPluginMap(const TSharedRef<FPlugin>&
 }
 #endif //if WITH_EDITOR
 
+void FPluginManager::GetExplanationForUnavailablePackage(const FName& UnavailablePackageName, FStringBuilderBase& InOutExplanation)
+{
+	FString UnavailablePackageNameString = UnavailablePackageName.ToString();
+	FPlugin* Plugin = static_cast<FPlugin*>(FindPluginFromPath(UnavailablePackageNameString).Get());
+
+	if (Plugin != nullptr)
+	{
+		if (!Plugin->IsMounted())
+		{
+			InOutExplanation.Appendf(TEXT("Skipped package %s was expected to be found in plugin %s, however that plugin is not mounted. [ExplicitlyLoaded = %s, Enabled = %s]."),
+				*UnavailablePackageNameString,
+				*Plugin->GetFriendlyName(),
+				Plugin->GetDescriptor().bExplicitlyLoaded ? TEXT("TRUE") : TEXT("FALSE"),
+				Plugin->IsEnabled() ? TEXT("TRUE") : TEXT("FALSE"));
+		}
+	}
+	else
+	{
+		// If we don't have a valid plugin, we have nothing to add
+	}
+
+	// Give higher level systems a chance to provide further information
+	GetExplanationForUnavailablePackageWithPluginInfoDelegate().Broadcast(UnavailablePackageNameString, Plugin, InOutExplanation);
+}
+
+IPluginManager::FGetExplanationForUnavailablePackageWithPluginInfoDelegate& FPluginManager::GetExplanationForUnavailablePackageWithPluginInfoDelegate()
+{
+	static FGetExplanationForUnavailablePackageWithPluginInfoDelegate Delegate;
+	return Delegate;
+}
+
+
 #undef LOCTEXT_NAMESPACE
diff --git a/Engine/Source/Runtime/Projects/Private/PluginManager.h b/Engine/Source/Runtime/Projects/Private/PluginManager.h
--- a/Engine/Source/Runtime/Projects/Private/PluginManager.h
+++ b/Engine/Source/Runtime/Projects/Private/PluginManager.h
@@ -145,6 +145,8 @@ public:
 	virtual void SetUnRegisterMountPointDelegate( const FRegisterMountPointDelegate& Delegate ) override;
 	virtual void SetUpdatePackageLocalizationCacheDelegate( const FUpdatePackageLocalizationCacheDelegate& Delegate ) override;
 	virtual bool AreRequiredPluginsAvailable() override;
+	virtual FGetExplanationForUnavailablePackageWithPluginInfoDelegate& GetExplanationForUnavailablePackageWithPluginInfoDelegate() override;
+
 #if !IS_MONOLITHIC
 	virtual bool CheckModuleCompatibility(TArray<FString>& OutIncompatibleModules, TArray<FString>& OutIncompatibleEngineModules) override;
 #endif
@@ -206,6 +208,8 @@ public:
 	virtual void SetPreloadBinaries() override;
 	virtual bool GetPreloadBinaries() override;
 
+	void GetExplanationForUnavailablePackage(const FName& UnavailablePackageName, FStringBuilderBase& InOutExplanation) override;
+
 private:
 	using FDiscoveredPluginMap = TMap<FString, TArray<TSharedRef<FPlugin>>>;
 
diff --git a/Engine/Source/Runtime/Projects/Public/Interfaces/IPluginManager.h b/Engine/Source/Runtime/Projects/Public/Interfaces/IPluginManager.h
--- a/Engine/Source/Runtime/Projects/Public/Interfaces/IPluginManager.h
+++ b/Engine/Source/Runtime/Projects/Public/Interfaces/IPluginManager.h
@@ -3,6 +3,7 @@
 #pragma once
 
 #include "CoreMinimal.h"
+#include "Delegates/Delegate.h"
 #include "PluginDescriptor.h"
 #include "Containers/VersePathFwd.h"
 #include "Templates/SharedPointer.h"
@@ -299,6 +300,25 @@ public:
 	 */
 	virtual bool AreRequiredPluginsAvailable() = 0;
 
+	/** Delegate type for allowing higher levels systems to provide additional context for why a package was skipped during loading */
+	DECLARE_MULTICAST_DELEGATE_ThreeParams(FGetExplanationForUnavailablePackageWithPluginInfoDelegate, const FString& /*Unavailable package*/, class IPlugin* /*Plugin or null if not found*/, FStringBuilderBase& /*InOutStringBuilder*/);
+
+	/**
+	 * Accessor for adding/removing FGetExplanationForUnavailablePackageWithPluginInfoDelegate delegates.
+	 * 
+	 * @returns the FGetExplanationForUnavailablePackageWithPluginInfoDelegate to which callbacks can be bound
+	 */
+	virtual FGetExplanationForUnavailablePackageWithPluginInfoDelegate& GetExplanationForUnavailablePackageWithPluginInfoDelegate() = 0;
+
+	/**
+	 * Attempts to fill out InOutExplanation with information from various sources (registered via GetExplanationForUnavailablePackageWithPluginInfoDelegate)
+	 * about why UnavailablePackageName is not available
+	 * 
+	 * @param UnavailablePackageName the name of the package about which to get information
+	 * @param InOutExplanation a stringbuilder to which to append explanatory information
+	 */
+	virtual	void GetExplanationForUnavailablePackage(const FName& UnavailablePackageName, FStringBuilderBase& InOutExplanation) = 0;
+
 #if !IS_MONOLITHIC
 	/** 
 	 * Checks whether modules for the enabled plug-ins are up to date.
