Add ensures to GFP subsytem to catch any GFP requests that happen before GFP policy is ready.

diff --git a/Engine/Plugins/Runtime/GameFeatures/Source/GameFeatures/Private/GameFeaturesSubsystem.cpp b/Engine/Plugins/Runtime/GameFeatures/Source/GameFeatures/Private/GameFeaturesSubsystem.cpp
--- a/Engine/Plugins/Runtime/GameFeatures/Source/GameFeatures/Private/GameFeaturesSubsystem.cpp
+++ b/Engine/Plugins/Runtime/GameFeatures/Source/GameFeatures/Private/GameFeaturesSubsystem.cpp
@@ -524,8 +524,14 @@ void UGameFeaturesSubsystem::OnAssetManagerCreated()
 
 	// Create the game-specific policy
 	UE_LOG(LogGameFeatures, Verbose, TEXT("Initializing game features policy (type %s)"), *GameSpecificPolicies->GetClass()->GetName());
+	bInitializedPolicyManager = true; // Set before calling InitGameFeatureManager() because InitGameFeatureManager may load GFPs
 	GameSpecificPolicies->InitGameFeatureManager();
-	bInitializedPolicyManager = true;
+}
+
+bool UGameFeaturesSubsystem::IsPluginAllowed(const FString& PluginURL) const
+{
+	ensureMsgf(bInitializedPolicyManager, TEXT("Attemting to load plugin [%s] before GameFeaturesSubsystem is ready!"), *PluginURL);
+	return bInitializedPolicyManager && GameSpecificPolicies->IsPluginAllowed(PluginURL);
 }
 
 TSharedPtr<FStreamableHandle> UGameFeaturesSubsystem::LoadGameFeatureData(const FString& GameFeatureToLoad, bool bStartStalled /*= false*/)
@@ -1218,7 +1224,7 @@ void UGameFeaturesSubsystem::LoadGameFeaturePlugin(const FString& PluginURL, con
 
 void UGameFeaturesSubsystem::LoadGameFeaturePlugin(const FString& PluginURL, const FGameFeatureProtocolOptions& ProtocolOptions, const FGameFeaturePluginLoadComplete& CompleteDelegate)
 {
-	const bool bIsPluginAllowed = GameSpecificPolicies->IsPluginAllowed(PluginURL);
+	const bool bIsPluginAllowed = IsPluginAllowed(PluginURL);
 	if (!bIsPluginAllowed)
 	{
 		CompleteDelegate.ExecuteIfBound(UE::GameFeatures::FResult(MakeError(UE::GameFeatures::SubsystemErrorNamespace + UE::GameFeatures::CommonErrorCodes::PluginNotAllowed)));
@@ -1327,7 +1333,7 @@ void UGameFeaturesSubsystem::ChangeGameFeatureTargetState(const FString& PluginU
 	static_assert(std::underlying_type<EGameFeatureTargetState>::type(EGameFeatureTargetState::Count) == 4, "");
 	check(TargetPluginState != EGameFeaturePluginState::MAX);
 
-	const bool bIsPluginAllowed = GameSpecificPolicies->IsPluginAllowed(PluginURL);
+	const bool bIsPluginAllowed = IsPluginAllowed(PluginURL);
 
 	UGameFeaturePluginStateMachine* StateMachine = nullptr;
 	if (!bIsPluginAllowed)
@@ -1749,7 +1755,7 @@ void UGameFeaturesSubsystem::LoadBuiltInGameFeaturePlugin(const TSharedRef<IPlug
 	FGameFeaturePluginDetails PluginDetails;
 	if (GetBuiltInGameFeaturePluginDetails(Plugin, PluginURL, PluginDetails))
 	{
-		if (GameSpecificPolicies->IsPluginAllowed(PluginURL))
+		if (IsPluginAllowed(PluginURL))
 		{
 			FBuiltInGameFeaturePluginBehaviorOptions BehaviorOptions;
 			const bool bShouldProcess = AdditionalFilter(Plugin->GetDescriptorFileName(), PluginDetails, BehaviorOptions);
diff --git a/Engine/Plugins/Runtime/GameFeatures/Source/GameFeatures/Public/GameFeaturesSubsystem.h b/Engine/Plugins/Runtime/GameFeatures/Source/GameFeatures/Public/GameFeaturesSubsystem.h
--- a/Engine/Plugins/Runtime/GameFeatures/Source/GameFeatures/Public/GameFeaturesSubsystem.h
+++ b/Engine/Plugins/Runtime/GameFeatures/Source/GameFeatures/Public/GameFeaturesSubsystem.h
@@ -588,6 +588,7 @@ public:
 	template <typename T = UGameFeaturesProjectPolicies>
 	T& GetPolicy() const
 	{
+		ensureMsgf(bInitializedPolicyManager, TEXT("Attemting to get policy before GameFeaturesSubsystem is ready!"));
 		return *CastChecked<T>(GameSpecificPolicies, ECastCheckedType::NullChecked);
 	}
 
@@ -697,6 +698,8 @@ private:
 	void GetExplanationForUnavailablePackage(const FString& SkippedPackage, IPlugin* PluginIfFound, FStringBuilderBase& InOutExplanation);
 
 private:
+	bool IsPluginAllowed(const FString& PluginURL) const;
+
 	bool ShouldUpdatePluginProtocolOptions(const UGameFeaturePluginStateMachine* StateMachine, const FGameFeatureProtocolOptions& NewOptions);
 	UE::GameFeatures::FResult UpdateGameFeatureProtocolOptions(UGameFeaturePluginStateMachine* StateMachine, const FGameFeatureProtocolOptions& NewOptions, bool* bOutDidUpdate = nullptr);
 
