Do not block or fail GFP unloading in FGameFeaturePluginState_DeactivatingDependencies if a dependency fails to unload.  In many cases it is valid to cancel the depencies transition if another system requests it to be loaded before fully unloading. Behavior can be reverted via GameFeaturePlugin.WaitForDependencyDeactivation CVar

diff --git a/Engine/Plugins/Runtime/GameFeatures/Source/GameFeatures/Private/GameFeaturePluginStateMachine.cpp b/Engine/Plugins/Runtime/GameFeatures/Source/GameFeatures/Private/GameFeaturePluginStateMachine.cpp
--- a/Engine/Plugins/Runtime/GameFeatures/Source/GameFeatures/Private/GameFeaturePluginStateMachine.cpp
+++ b/Engine/Plugins/Runtime/GameFeatures/Source/GameFeatures/Private/GameFeaturePluginStateMachine.cpp
@@ -11,6 +11,7 @@
 #include "InstallBundleUtils.h"
 #include "BundlePrereqCombinedStatusHelper.h"
 #include "Interfaces/IPluginManager.h"
+#include "Logging/StructuredLog.h"
 #include "Materials/MaterialInterface.h"
 #include "Misc/AsciiSet.h"
 #include "Misc/ConfigCacheIni.h"
@@ -81,6 +82,10 @@ namespace UE::GameFeatures
 		false,
 		TEXT("Enable to force calls to IAssetRegistry::AppendState to happen on the game thread"));
 
+	static TAutoConsoleVariable<bool> CVarWaitForDependencyDeactivation(TEXT("GameFeaturePlugin.WaitForDependencyDeactivation"),
+		false,
+		TEXT("Enable to make block deactivation until all dependencies are deactivated. Warning - this can lead to failure to unload"));
+
 	bool ShouldSkipVerify(const FString& PluginName)
 	{
 		static const FAsciiSet Wildcards("*?");
@@ -534,8 +539,22 @@ struct FTransitionDependenciesGameFeaturePluginState : public FGameFeaturePlugin
 
 	void TransitionDependency(UGameFeaturePluginStateMachine* Dependency)
 	{
-		const bool bSetDestination = Dependency->SetDestination(TransitionPolicy::GetDependencyStateRange(),
-			FGameFeatureStateTransitionComplete::CreateRaw(this, &FTransitionDependenciesGameFeaturePluginState::OnDependencyTransitionComplete));
+		bool bSetDestination = false;
+		
+		if (TransitionPolicy::ShouldWaitForDependencies())
+		{
+			bSetDestination = Dependency->SetDestination(TransitionPolicy::GetDependencyStateRange(),
+				FGameFeatureStateTransitionComplete::CreateRaw(this, &FTransitionDependenciesGameFeaturePluginState::OnDependencyTransitionComplete));
+		}
+		else
+		{
+			bSetDestination = Dependency->SetDestination(TransitionPolicy::GetDependencyStateRange(), 
+				FGameFeatureStateTransitionComplete::CreateStatic(&FTransitionDependenciesGameFeaturePluginState::OnDependencyTransitionCompleteNoWait));
+			if (bSetDestination)
+			{
+				OnDependencyTransitionComplete(Dependency, MakeValue());
+			}
+		}
 
 		if (!bSetDestination)
 		{
@@ -580,8 +599,21 @@ struct FTransitionDependenciesGameFeaturePluginState : public FGameFeaturePlugin
 		}
 
 		// Now that the transition has been canceled, retry reaching the desired destination
-		const bool bSetDestination = Dependency->SetDestination(TransitionPolicy::GetDependencyStateRange(),
-			FGameFeatureStateTransitionComplete::CreateRaw(this, &FTransitionDependenciesGameFeaturePluginState::OnDependencyTransitionComplete));
+		bool bSetDestination = false;
+		if (TransitionPolicy::ShouldWaitForDependencies())
+		{
+			bSetDestination = Dependency->SetDestination(TransitionPolicy::GetDependencyStateRange(),
+				FGameFeatureStateTransitionComplete::CreateRaw(this, &FTransitionDependenciesGameFeaturePluginState::OnDependencyTransitionComplete));
+		}
+		else
+		{
+			bSetDestination = Dependency->SetDestination(TransitionPolicy::GetDependencyStateRange(), 
+				FGameFeatureStateTransitionComplete::CreateStatic(&FTransitionDependenciesGameFeaturePluginState::OnDependencyTransitionCompleteNoWait));
+			if (bSetDestination)
+			{
+				OnDependencyTransitionComplete(Dependency, MakeValue());
+			}
+		}
 
 		if (!ensure(bSetDestination))
 		{
@@ -613,6 +645,26 @@ struct FTransitionDependenciesGameFeaturePluginState : public FGameFeaturePlugin
 		}
 	}
 
+	static void OnDependencyTransitionCompleteNoWait(UGameFeaturePluginStateMachine* Dependency, const UE::GameFeatures::FResult& Result)
+	{
+		if (Result.HasError())
+		{
+			if (Result.GetError() == UE::GameFeatures::CanceledResult.GetError())
+			{
+				UE_LOGFMT(LogGameFeatures, Warning, "Dependency {Dep} failed to transition because it was cancelled by another request {Error}",
+					("Dep", Dependency->GetPluginIdentifier().GetIdentifyingString()),
+					("Error", Result.GetError()));
+			}
+			else
+			{
+
+				UE_LOGFMT(LogGameFeatures, Error, "Dependency {Dep} failed to transition with error {Error}",
+					("Dep", Dependency->GetPluginIdentifier().GetIdentifyingString()),
+					("Error", Result.GetError()));
+			}
+		}
+	}
+
 	void ClearDependencies()
 	{
 		for (FDepResultPair& Pair : RemainingDependencies)
@@ -2161,7 +2213,7 @@ struct FWaitingForDependenciesTransitionPolicy
 		UGameFeaturesSubsystem& GameFeaturesSubsystem = UGameFeaturesSubsystem::Get();
 
 		return GameFeaturesSubsystem.FindOrCreatePluginDependencyStateMachines(
-			*InStateProperties.PluginIdentifier.GetFullPluginURL(), InStateProperties, OutDependencyMachines);
+			InStateProperties.PluginIdentifier.GetFullPluginURL(), InStateProperties, OutDependencyMachines);
 	}
 
 	static FGameFeaturePluginStateRange GetDependencyStateRange()
@@ -2178,6 +2230,11 @@ struct FWaitingForDependenciesTransitionPolicy
 	{
 		return EGameFeaturePluginState::ErrorWaitingForDependencies;
 	}
+
+	static bool ShouldWaitForDependencies()
+	{
+		return true;
+	}
 };
 
 struct FGameFeaturePluginState_WaitingForDependencies : public FTransitionDependenciesGameFeaturePluginState<FWaitingForDependenciesTransitionPolicy>
@@ -2677,7 +2734,7 @@ struct FDeactivatingDependenciesTransitionPolicy
 		UGameFeaturesSubsystem& GameFeaturesSubsystem = UGameFeaturesSubsystem::Get();
 
 		return GameFeaturesSubsystem.FindPluginDependencyStateMachinesToDeactivate(
-			*InStateProperties.PluginIdentifier.GetFullPluginURL(), InStateProperties.PluginInstalledFilename, OutDependencyMachines);
+			InStateProperties.PluginIdentifier.GetFullPluginURL(), InStateProperties.PluginInstalledFilename, OutDependencyMachines);
 	}
 
 	static FGameFeaturePluginStateRange GetDependencyStateRange()
@@ -2694,6 +2751,11 @@ struct FDeactivatingDependenciesTransitionPolicy
 	{
 		return EGameFeaturePluginState::ErrorDeactivatingDependencies;
 	}
+
+	static bool ShouldWaitForDependencies()
+	{
+		return UE::GameFeatures::CVarWaitForDependencyDeactivation.GetValueOnGameThread();
+	}
 };
 
 struct FGameFeaturePluginState_DeactivatingDependencies : public FTransitionDependenciesGameFeaturePluginState<FDeactivatingDependenciesTransitionPolicy>
@@ -2817,7 +2879,7 @@ struct FActivatingDependenciesTransitionPolicy
 		UGameFeaturesSubsystem& GameFeaturesSubsystem = UGameFeaturesSubsystem::Get();
 
 		return GameFeaturesSubsystem.FindPluginDependencyStateMachinesToActivate(
-			*InStateProperties.PluginIdentifier.GetFullPluginURL(), InStateProperties.PluginInstalledFilename, OutDependencyMachines);
+			InStateProperties.PluginIdentifier.GetFullPluginURL(), InStateProperties.PluginInstalledFilename, OutDependencyMachines);
 	}
 
 	static FGameFeaturePluginStateRange GetDependencyStateRange()
@@ -2834,6 +2896,11 @@ struct FActivatingDependenciesTransitionPolicy
 	{
 		return EGameFeaturePluginState::ErrorActivatingDependencies;
 	}
+
+	static bool ShouldWaitForDependencies()
+	{
+		return true;
+	}
 };
 
 struct FGameFeaturePluginState_ActivatingDependencies : public FTransitionDependenciesGameFeaturePluginState<FActivatingDependenciesTransitionPolicy>
