Cleanup of GFP command line debug tools.
- Fix not being able to pass full plugin URLs
- Add a utility to do argument parsing

diff --git a/Engine/Plugins/Runtime/GameFeatures/Source/GameFeatures/Private/GameFeaturesSubsystem.cpp b/Engine/Plugins/Runtime/GameFeatures/Source/GameFeatures/Private/GameFeaturesSubsystem.cpp
--- a/Engine/Plugins/Runtime/GameFeatures/Source/GameFeatures/Private/GameFeaturesSubsystem.cpp
+++ b/Engine/Plugins/Runtime/GameFeatures/Source/GameFeatures/Private/GameFeaturesSubsystem.cpp
@@ -96,6 +96,38 @@ namespace UE::GameFeatures
 		TEXT("When listing worst offenders for built-in plugin load time, show no more than this many plugins, to reduce log spam."),
 		ECVF_Default);
 #endif // !UE_BUILD_SHIPPING
+
+	TOptional<FString> GetPluginUrlForConsoleCommand(const TArray<FString>& Args, FOutputDevice& Ar)
+	{
+		TOptional<FString> PluginURL;
+		if (Args.Num() > 0)
+		{
+			EGameFeaturePluginProtocol Protocol = UGameFeaturesSubsystem::GetPluginURLProtocol(Args[0]);
+			if (Protocol != EGameFeaturePluginProtocol::Unknown)
+			{
+				PluginURL.Emplace(Args[0]);
+			}
+			else
+			{
+				FString PluginURLStr;
+				if (UGameFeaturesSubsystem::Get().GetPluginURLByName(Args[0], /*out*/ PluginURLStr))
+				{
+					PluginURL.Emplace(MoveTemp(PluginURLStr));
+				}
+			}
+		}
+
+		if (PluginURL)
+		{
+			Ar.Logf(TEXT("Using URL %s for console command"), *PluginURL.GetValue());
+		}
+		else
+		{
+			Ar.Logf(TEXT("Expected a game feature plugin URL or name as an argument"));
+		}
+
+		return PluginURL;
+	}
 }
 
 const FString LexToString(const EBuiltInAutoState BuiltInAutoState)
@@ -350,18 +382,9 @@ void UGameFeaturesSubsystem::Initialize(FSubsystemCollectionBase& Collection)
 		TEXT("Loads and activates a game feature plugin by PluginName or URL"),
 		FConsoleCommandWithWorldArgsAndOutputDeviceDelegate::CreateLambda([](const TArray<FString>& Args, UWorld*, FOutputDevice& Ar)
 		{
-			if (Args.Num() > 0)
+			if (TOptional<FString> PluginURL = UE::GameFeatures::GetPluginUrlForConsoleCommand(Args, Ar))
 			{
-				FString PluginURL;
-				if (!UGameFeaturesSubsystem::Get().GetPluginURLByName(Args[0], /*out*/ PluginURL))
-				{
-					PluginURL = Args[0];
-				}
-				UGameFeaturesSubsystem::Get().LoadAndActivateGameFeaturePlugin(PluginURL, FGameFeaturePluginLoadComplete());
-			}
-			else
-			{
-				Ar.Logf(TEXT("Expected a game feature plugin URL as an argument"));
+				UGameFeaturesSubsystem::Get().LoadAndActivateGameFeaturePlugin(PluginURL.GetValue(), FGameFeaturePluginLoadComplete());
 			}
 		}),
 		ECVF_Cheat);
@@ -371,18 +394,9 @@ void UGameFeaturesSubsystem::Initialize(FSubsystemCollectionBase& Collection)
 		TEXT("Deactivates a game feature plugin by PluginName or URL"),
 		FConsoleCommandWithWorldArgsAndOutputDeviceDelegate::CreateLambda([](const TArray<FString>& Args, UWorld*, FOutputDevice& Ar)
 		{
-			if (Args.Num() > 0)
-			{
-				FString PluginURL;
-				if (!UGameFeaturesSubsystem::Get().GetPluginURLByName(Args[0], /*out*/ PluginURL))
-				{
-					PluginURL = Args[0];
-				}
-				UGameFeaturesSubsystem::Get().DeactivateGameFeaturePlugin(PluginURL, FGameFeaturePluginLoadComplete());
-			}
-			else
+			if (TOptional<FString> PluginURL = UE::GameFeatures::GetPluginUrlForConsoleCommand(Args, Ar))
 			{
-				Ar.Logf(TEXT("Expected a game feature plugin URL as an argument"));
+				UGameFeaturesSubsystem::Get().DeactivateGameFeaturePlugin(PluginURL.GetValue(), FGameFeaturePluginLoadComplete());
 			}
 		}),
 		ECVF_Cheat);
@@ -392,18 +406,9 @@ void UGameFeaturesSubsystem::Initialize(FSubsystemCollectionBase& Collection)
 		TEXT("Unloads a game feature plugin by PluginName or URL"),
 		FConsoleCommandWithWorldArgsAndOutputDeviceDelegate::CreateLambda([](const TArray<FString>& Args, UWorld*, FOutputDevice& Ar)
 		{
-			if (Args.Num() > 0)
+			if (TOptional<FString> PluginURL = UE::GameFeatures::GetPluginUrlForConsoleCommand(Args, Ar))
 			{
-				FString PluginURL;
-				if (!UGameFeaturesSubsystem::Get().GetPluginURLByName(Args[0], /*out*/ PluginURL))
-				{
-					PluginURL = Args[0];
-				}
-				UGameFeaturesSubsystem::Get().UnloadGameFeaturePlugin(PluginURL);
-			}
-			else
-			{
-				Ar.Logf(TEXT("Expected a game feature plugin URL as an argument"));
+				UGameFeaturesSubsystem::Get().UnloadGameFeaturePlugin(PluginURL.GetValue());
 			}
 		}),
 		ECVF_Cheat);
@@ -412,84 +417,48 @@ void UGameFeaturesSubsystem::Initialize(FSubsystemCollectionBase& Collection)
 		TEXT("UnloadAndKeepRegisteredGameFeaturePlugin"),
 		TEXT("Unloads a game feature plugin by PluginName or URL but keeps it registered"),
 		FConsoleCommandWithWorldArgsAndOutputDeviceDelegate::CreateLambda([](const TArray<FString>& Args, UWorld*, FOutputDevice& Ar)
+		{
+			if (TOptional<FString> PluginURL = UE::GameFeatures::GetPluginUrlForConsoleCommand(Args, Ar))
 			{
-				if (Args.Num() > 0)
-				{
-					FString PluginURL;
-					if (!UGameFeaturesSubsystem::Get().GetPluginURLByName(Args[0], /*out*/ PluginURL))
-					{
-						PluginURL = Args[0];
-					}
-					UGameFeaturesSubsystem::Get().UnloadGameFeaturePlugin(PluginURL, true);
-				}
-				else
-				{
-					Ar.Logf(TEXT("Expected a game feature plugin URL as an argument"));
-				}
-			}),
+				UGameFeaturesSubsystem::Get().UnloadGameFeaturePlugin(PluginURL.GetValue(), true);
+			}
+		}),
 		ECVF_Cheat);
 
 	IConsoleManager::Get().RegisterConsoleCommand(
 		TEXT("ReleaseGameFeaturePlugin"),
 		TEXT("Releases a game feature plugin's InstallBundle data by PluginName or URL"),
 		FConsoleCommandWithWorldArgsAndOutputDeviceDelegate::CreateLambda([](const TArray<FString>& Args, UWorld*, FOutputDevice& Ar)
+		{
+			if (TOptional<FString> PluginURL = UE::GameFeatures::GetPluginUrlForConsoleCommand(Args, Ar))
 			{
-				if (Args.Num() > 0)
-				{
-					FString PluginURL;
-					if (!UGameFeaturesSubsystem::Get().GetPluginURLByName(Args[0], /*out*/ PluginURL))
-					{
-						PluginURL = Args[0];
-					}
-					UGameFeaturesSubsystem::Get().ReleaseGameFeaturePlugin(PluginURL);
-				}
-				else
-				{
-					Ar.Logf(TEXT("Expected a game feature plugin URL as an argument"));
-				}
-			}),
+				UGameFeaturesSubsystem::Get().ReleaseGameFeaturePlugin(PluginURL.GetValue());
+			}
+		}),
 		ECVF_Cheat);
 
 	IConsoleManager::Get().RegisterConsoleCommand(
 		TEXT("CancelGameFeaturePlugin"),
 		TEXT("Cancel any state changes for a game feature plugin by PluginName or URL"),
 		FConsoleCommandWithWorldArgsAndOutputDeviceDelegate::CreateLambda([](const TArray<FString>& Args, UWorld*, FOutputDevice& Ar)
+		{
+			if (TOptional<FString> PluginURL = UE::GameFeatures::GetPluginUrlForConsoleCommand(Args, Ar))
 			{
-				if (Args.Num() > 0)
-				{
-					FString PluginURL;
-					if (!UGameFeaturesSubsystem::Get().GetPluginURLByName(Args[0], /*out*/ PluginURL))
-					{
-						PluginURL = Args[0];
-					}
-					UGameFeaturesSubsystem::Get().CancelGameFeatureStateChange(PluginURL);
-				}
-				else
-				{
-					Ar.Logf(TEXT("Expected a game feature plugin URL as an argument"));
-				}
-			}),
+				UGameFeaturesSubsystem::Get().CancelGameFeatureStateChange(PluginURL.GetValue());
+			}
+		}),
 		ECVF_Cheat);
 
 	IConsoleManager::Get().RegisterConsoleCommand(
 		TEXT("TerminateGameFeaturePlugin"),
 		TEXT("Terminates a game feature plugin by PluginName or URL"),
 		FConsoleCommandWithWorldArgsAndOutputDeviceDelegate::CreateLambda([](const TArray<FString>& Args, UWorld*, FOutputDevice& Ar)
+		{
+			if (TOptional<FString> PluginURL = UE::GameFeatures::GetPluginUrlForConsoleCommand(Args, Ar))
 			{
-				if (Args.Num() > 0)
-				{
-					FString PluginURL;
-					if (!UGameFeaturesSubsystem::Get().GetPluginURLByName(Args[0], /*out*/ PluginURL))
-					{
-						PluginURL = Args[0];
-					}
-					UGameFeaturesSubsystem::Get().TerminateGameFeaturePlugin(PluginURL);
-				}
-				else
-				{
-					Ar.Logf(TEXT("Expected a game feature plugin URL as an argument"));
-				}
-			}),
+				UGameFeaturesSubsystem::Get().TerminateGameFeaturePlugin(PluginURL.GetValue());
+			}
+		}),
 		ECVF_Cheat);
 
 	GetExplanationForUnavailablePackageDelegateHandle = IPluginManager::Get().GetExplanationForUnavailablePackageWithPluginInfoDelegate().AddUObject(this, &UGameFeaturesSubsystem::GetExplanationForUnavailablePackage);
