Fixed crash in shutdown if game feature plugin was unloaded/destroyed during shutdown and the tags manager had already been destroyed.

diff --git a/Engine/Plugins/Runtime/GameFeatures/Source/GameFeatures/Private/GameFeaturesSubsystem.cpp b/Engine/Plugins/Runtime/GameFeatures/Source/GameFeatures/Private/GameFeaturesSubsystem.cpp
--- a/Engine/Plugins/Runtime/GameFeatures/Source/GameFeatures/Private/GameFeaturesSubsystem.cpp
+++ b/Engine/Plugins/Runtime/GameFeatures/Source/GameFeatures/Private/GameFeaturesSubsystem.cpp
@@ -1233,10 +1233,31 @@ void UGameFeaturesSubsystem::LoadGameFeaturePlugin(TConstArrayView<FString> Plug
 		FMultipleGameFeaturePluginsLoaded CompleteDelegate;
 
 		int32 NumPluginsLoaded = 0;
+		bool bPushedTagsBroadcast = false;
+
+		FLoadContext()
+		{
+			if (!IsEngineExitRequested())
+			{
+				UGameplayTagsManager::Get().PushDeferOnGameplayTagTreeChangedBroadcast();
+				bPushedTagsBroadcast = true;
+			}
+			else if(UGameplayTagsManager* TagsManager = UGameplayTagsManager::GetIfAllocated())
+			{
+				TagsManager->PushDeferOnGameplayTagTreeChangedBroadcast();
+				bPushedTagsBroadcast = true;
+			}
+		}
 
 		~FLoadContext()
 		{
-			UGameplayTagsManager::Get().PopDeferOnGameplayTagTreeChangedBroadcast();
+			if (bPushedTagsBroadcast)
+			{
+				if (UGameplayTagsManager* TagsManager = UGameplayTagsManager::GetIfAllocated())
+				{
+					TagsManager->PopDeferOnGameplayTagTreeChangedBroadcast();
+				}
+			}
 
 			CompleteDelegate.ExecuteIfBound(Results);
 		}
@@ -1244,8 +1265,6 @@ void UGameFeaturesSubsystem::LoadGameFeaturePlugin(TConstArrayView<FString> Plug
 	TSharedRef<FLoadContext> LoadContext = MakeShared<FLoadContext>();
 	LoadContext->CompleteDelegate = CompleteDelegate;
 
-	UGameplayTagsManager::Get().PushDeferOnGameplayTagTreeChangedBroadcast();
-
 	LoadContext->Results.Reserve(PluginURLs.Num());
 	for (const FString& PluginURL : PluginURLs)
 	{
@@ -1367,10 +1386,31 @@ void UGameFeaturesSubsystem::ChangeGameFeatureTargetState(TConstArrayView<FStrin
 		FMultipleGameFeaturePluginsLoaded CompleteDelegate;
 
 		int32 NumPluginsLoaded = 0;
+		bool bPushedTagsBroadcast = false;
+
+		FLoadContext()
+		{
+			if (!IsEngineExitRequested())
+			{
+				UGameplayTagsManager::Get().PushDeferOnGameplayTagTreeChangedBroadcast();
+				bPushedTagsBroadcast = true;
+			}
+			else if(UGameplayTagsManager* TagsManager = UGameplayTagsManager::GetIfAllocated())
+			{
+				TagsManager->PushDeferOnGameplayTagTreeChangedBroadcast();
+				bPushedTagsBroadcast = true;
+			}
+		}
 
 		~FLoadContext()
 		{
-			UGameplayTagsManager::Get().PopDeferOnGameplayTagTreeChangedBroadcast();
+			if (bPushedTagsBroadcast)
+			{
+				if (UGameplayTagsManager* TagsManager = UGameplayTagsManager::GetIfAllocated())
+				{
+					TagsManager->PopDeferOnGameplayTagTreeChangedBroadcast();
+				}
+			}
 
 			CompleteDelegate.ExecuteIfBound(Results);
 		}
@@ -1378,8 +1418,6 @@ void UGameFeaturesSubsystem::ChangeGameFeatureTargetState(TConstArrayView<FStrin
 	TSharedRef<FLoadContext> LoadContext = MakeShared<FLoadContext>();
 	LoadContext->CompleteDelegate = CompleteDelegate;
 
-	UGameplayTagsManager::Get().PushDeferOnGameplayTagTreeChangedBroadcast();
-
 	LoadContext->Results.Reserve(PluginURLs.Num());
 	for (const FString& PluginURL : PluginURLs)
 	{
@@ -1861,11 +1899,49 @@ void UGameFeaturesSubsystem::LoadBuiltInGameFeaturePlugins(FBuiltInPluginAdditio
 		FBuiltInGameFeaturePluginsLoaded CompleteDelegate;
 
 		int32 NumPluginsLoaded = 0;
+		bool bPushedTagsBroadcast = false;
+		bool bPushedAssetBulkScanning = false;
+
+		FLoadContext()
+		{
+			if (!IsEngineExitRequested())
+			{
+				UAssetManager::Get().PushBulkScanning();
+				bPushedAssetBulkScanning = true;
+				UGameplayTagsManager::Get().PushDeferOnGameplayTagTreeChangedBroadcast();
+				bPushedTagsBroadcast = true;
+			}
+			else
+			{
+				if(UAssetManager* AssetManager =  UAssetManager::GetIfInitialized())
+				{
+					AssetManager->PushBulkScanning();
+					bPushedAssetBulkScanning = true;
+				}
+				if(UGameplayTagsManager* TagsManager = UGameplayTagsManager::GetIfAllocated())
+				{
+					TagsManager->PushDeferOnGameplayTagTreeChangedBroadcast();
+					bPushedTagsBroadcast = true;
+				}
+			}
+		}
 
 		~FLoadContext()
 		{
-			UGameplayTagsManager::Get().PopDeferOnGameplayTagTreeChangedBroadcast();
-			UAssetManager::Get().PopBulkScanning();
+			if (bPushedTagsBroadcast)
+			{
+				if (UGameplayTagsManager* TagsManager = UGameplayTagsManager::GetIfAllocated())
+				{
+					TagsManager->PopDeferOnGameplayTagTreeChangedBroadcast();
+				}
+			}
+			if (bPushedAssetBulkScanning)
+			{
+				if(UAssetManager* AssetManager =  UAssetManager::GetIfInitialized())
+				{
+					AssetManager->PopBulkScanning();
+				}
+			}
 
 			CompleteDelegate.ExecuteIfBound(Results);
 		}
@@ -1873,9 +1949,6 @@ void UGameFeaturesSubsystem::LoadBuiltInGameFeaturePlugins(FBuiltInPluginAdditio
 	TSharedRef<FLoadContext> LoadContext = MakeShared<FLoadContext>();
 	LoadContext->CompleteDelegate = InCompleteDelegate;
 
-	UAssetManager::Get().PushBulkScanning();
-	UGameplayTagsManager::Get().PushDeferOnGameplayTagTreeChangedBroadcast();
-
 	FBuiltInPluginLoadTimeTracker PluginLoadTimeTracker;
 	TArray<TSharedRef<IPlugin>> EnabledPlugins = IPluginManager::Get().GetEnabledPlugins();
 
