Stub out a GFP state for streaming asset dependencies.  Can be enabled with GameFeaturePlugin.EnableAssetStreaming CVar.
FGameFeaturePluginDetails now contains asset dependencies if present.

diff --git a/Engine/Plugins/Runtime/GameFeatures/Source/GameFeatures/Private/GameFeaturePluginStateMachine.cpp b/Engine/Plugins/Runtime/GameFeatures/Source/GameFeatures/Private/GameFeaturePluginStateMachine.cpp
--- a/Engine/Plugins/Runtime/GameFeatures/Source/GameFeatures/Private/GameFeaturePluginStateMachine.cpp
+++ b/Engine/Plugins/Runtime/GameFeatures/Source/GameFeatures/Private/GameFeaturePluginStateMachine.cpp
@@ -86,6 +86,10 @@ namespace UE::GameFeatures
 		false,
 		TEXT("Enable to make block deactivation until all dependencies are deactivated. Warning - this can lead to failure to unload"));
 
+	static TAutoConsoleVariable<bool> CVarEnableAssetStreaming(TEXT("GameFeaturePlugin.EnableAssetStreaming"),
+		false,
+		TEXT("Enable experimental asset streaming"));
+
 	bool ShouldSkipVerify(const FString& PluginName)
 	{
 		static const FAsciiSet Wildcards("*?");
@@ -2223,7 +2227,7 @@ struct FWaitingForDependenciesTransitionPolicy
 
 	static EGameFeaturePluginState GetTransitionState()
 	{
-		return EGameFeaturePluginState::Registering;
+		return UE::GameFeatures::CVarEnableAssetStreaming.GetValueOnGameThread() ? EGameFeaturePluginState::AssetDependencyStreaming : EGameFeaturePluginState::Registering;
 	}
 
 	static EGameFeaturePluginState GetErrorState()
@@ -2245,6 +2249,39 @@ struct FGameFeaturePluginState_WaitingForDependencies : public FTransitionDepend
 	}
 };
 
+struct FGameFeaturePluginState_AssetDependencyStreaming : public FGameFeaturePluginState
+{
+	FGameFeaturePluginState_AssetDependencyStreaming(FGameFeaturePluginStateMachineProperties& InStateProperties) : FGameFeaturePluginState(InStateProperties) {}
+
+	TArray<FString> AssetDependencies;
+
+	virtual void BeginState() override
+	{
+		AssetDependencies = UGameFeaturesSubsystem::Get().FindPluginAssetDependencies(StateProperties.PluginInstalledFilename);
+	}
+
+	virtual void UpdateState(FGameFeaturePluginStateStatus& StateStatus) override
+	{
+		StateStatus.SetTransition(EGameFeaturePluginState::Registering);
+	}
+
+	virtual void EndState() override
+	{
+		AssetDependencies.Empty();
+	}
+
+	virtual void TryCancelState() override
+	{
+		// TODO: support pause, cancel etc.  We could possibly be streaming hundreds of assets.
+	}
+
+	virtual UE::GameFeatures::FResult TryUpdateProtocolOptions(const FGameFeatureProtocolOptions& NewOptions) override
+	{
+		// TODO: support pause, cancel etc.  We could possibly be streaming hundreds of assets.
+		return FGameFeaturePluginState::TryUpdateProtocolOptions(NewOptions);
+	}
+};
+
 struct FGameFeaturePluginState_Unregistering : public FGameFeaturePluginState
 {
 	FGameFeaturePluginState_Unregistering(FGameFeaturePluginStateMachineProperties& InStateProperties) : FGameFeaturePluginState(InStateProperties) {}
diff --git a/Engine/Plugins/Runtime/GameFeatures/Source/GameFeatures/Private/GameFeaturePluginStateMachine.h b/Engine/Plugins/Runtime/GameFeatures/Source/GameFeatures/Private/GameFeaturePluginStateMachine.h
--- a/Engine/Plugins/Runtime/GameFeatures/Source/GameFeatures/Private/GameFeaturePluginStateMachine.h
+++ b/Engine/Plugins/Runtime/GameFeatures/Source/GameFeatures/Private/GameFeaturePluginStateMachine.h
@@ -100,11 +100,17 @@ Transition states are expected to transition the machine to another state after
            |     |          |
            ------~----------~-------------------------------
                  |          |                              |
-                 |       +--v------------------- +   +-----v-----------------------+
-                 |       |                       |	 |         !                   |
+                 |       +--v--------------------+   +-----v-----------------------+
+                 |       |                       |   |         !                   |
                  |       |WaitingForDependencies <---> ErrorWaitingForDependencies |
-                 |       |                       |	 |                             |
-                 |       +--+------------------- +	 +-----------------------------+
+                 |       |                       |   |                             |
+                 |       +--+--------------------+   +-----------^-----------------+
+                 |          |                                    |
+                 |       +--v----------------------+             |
+                 |       |                         |             |
+                 |       |AssetDependencyStreaming +--------------
+                 |       |                         |
+                 |       +--+----------------------+
                  |          |
            ------~----------~-------------------------------
            |     |          |                              |
diff --git a/Engine/Plugins/Runtime/GameFeatures/Source/GameFeatures/Private/GameFeaturesSubsystem.cpp b/Engine/Plugins/Runtime/GameFeatures/Source/GameFeatures/Private/GameFeaturesSubsystem.cpp
--- a/Engine/Plugins/Runtime/GameFeatures/Source/GameFeatures/Private/GameFeaturesSubsystem.cpp
+++ b/Engine/Plugins/Runtime/GameFeatures/Source/GameFeatures/Private/GameFeaturesSubsystem.cpp
@@ -2162,15 +2162,15 @@ bool UGameFeaturesSubsystem::GetGameFeaturePluginDetailsInternal(const FString&
 			FString FileContents;
 			if (!FFileHelper::LoadFileToString(FileContents, *PluginDescriptorFilename))
 			{
-				UE_LOG(LogGameFeatures, Error, TEXT("UGameFeaturesSubsystem could not determine if feature was hotfixable. Failed to read file. File:%s Error:%d"), *PluginDescriptorFilename, FPlatformMisc::GetLastError());
+				UE_LOG(LogGameFeatures, Error, TEXT("UGameFeaturesSubsystem could not load plugin descriptor. Failed to read file. File:%s Error:%d"), *PluginDescriptorFilename, FPlatformMisc::GetLastError());
 				return false;
 			}
 
 			// Deserialize a JSON object from the string	
-			TSharedRef< TJsonReader<> > Reader = TJsonReaderFactory<>::Create(FileContents);
+			TSharedRef<TJsonReader<>> Reader = TJsonReaderFactory<>::Create(FileContents);
 			if (!FJsonSerializer::Deserialize(Reader, ObjectPtr) || !ObjectPtr.IsValid())
 			{
-				UE_LOG(LogGameFeatures, Error, TEXT("UGameFeaturesSubsystem could not determine if feature was hotfixable. Json invalid. File:%s. Error:%s"), *PluginDescriptorFilename, *Reader->GetErrorMessage());
+				UE_LOG(LogGameFeatures, Error, TEXT("UGameFeaturesSubsystem could not load plugin descriptor. Json invalid. File:%s. Error:%s"), *PluginDescriptorFilename, *Reader->GetErrorMessage());
 				return false;
 			}
 		}
@@ -2179,7 +2179,7 @@ bool UGameFeaturesSubsystem::GetGameFeaturePluginDetailsInternal(const FString&
 	// Read the properties
 	// Hotfixable. If it is not specified, then we assume it is
 	OutPluginDetails.bHotfixable = true;
-	ObjectPtr->TryGetBoolField(TEXT("Hotfixable"), OutPluginDetails.bHotfixable);
+	ObjectPtr->TryGetBoolField(TEXTVIEW("Hotfixable"), OutPluginDetails.bHotfixable);
 
 	// Determine the initial plugin state
 	OutPluginDetails.BuiltInAutoState = DetermineBuiltInInitialFeatureState(ObjectPtr, PluginDescriptorFilename);
@@ -2200,47 +2200,67 @@ bool UGameFeaturesSubsystem::GetGameFeaturePluginDetailsInternal(const FString&
 
 	// Parse plugin dependencies
 	const TArray<TSharedPtr<FJsonValue>>* PluginsArray = nullptr;
-	ObjectPtr->TryGetArrayField(TEXT("Plugins"), PluginsArray);
+	ObjectPtr->TryGetArrayField(TEXTVIEW("Plugins"), PluginsArray);
 	if (PluginsArray)
 	{
-		FString NameField = TEXT("Name");
-		FString EnabledField = TEXT("Enabled");
-		FString ActivateField = TEXT("Activate");
+		const FStringView NameField = TEXTVIEW("Name");
+		const FStringView EnabledField = TEXTVIEW("Enabled");
+		const FStringView ActivateField = TEXTVIEW("Activate");
+		const FStringView AssetReferencesField = TEXTVIEW("AssetReferences");
 		for (const TSharedPtr<FJsonValue>& PluginElement : *PluginsArray)
 		{
-			if (PluginElement.IsValid())
+			if (!PluginElement)
 			{
-				const TSharedPtr<FJsonObject>* ElementObjectPtr = nullptr;
-				PluginElement->TryGetObject(ElementObjectPtr);
-				if (ElementObjectPtr && ElementObjectPtr->IsValid())
-				{
-					const TSharedPtr<FJsonObject>& ElementObject = *ElementObjectPtr;
+				continue;
+			}
 
-					FString DependencyName;
-					ElementObject->TryGetStringField(NameField, DependencyName);
-					if (!DependencyName.IsEmpty())
-					{
-						bool bElementEnabled = false;
-						ElementObject->TryGetBoolField(EnabledField, bElementEnabled);
-						if (bElementEnabled)
-						{
-							//Have to get Activate from JSON as it's unique to GFP and not in the PluginManager
-							bool bElementActivate = false;
-							ElementObject->TryGetBoolField(ActivateField, bElementActivate);
-
-							OutPluginDetails.PluginDependencies.Emplace(FGameFeaturePluginReferenceDetails(MoveTemp(DependencyName), bElementActivate));
-						}
-						else
-						{
-							UE_LOG(LogGameFeatures, VeryVerbose, TEXT("Skipping adding dependency %s in %s. Plugin is disabled."), *DependencyName, *PluginDescriptorFilename);
-						}
-					}
-					else
+			const TSharedPtr<FJsonObject>* ElementObjectPtr = nullptr;
+			PluginElement->TryGetObject(ElementObjectPtr);
+			if (!ElementObjectPtr || !ElementObjectPtr->IsValid())
+			{
+				continue;
+			}
+			const TSharedPtr<FJsonObject>& ElementObject = *ElementObjectPtr;
+
+			FString DependencyName;
+			ElementObject->TryGetStringField(NameField, DependencyName);
+			if (DependencyName.IsEmpty())
+			{
+				UE_LOG(LogGameFeatures, Error, TEXT("Error parsing dependency name in %s! Invalid JSON data!"), *PluginDescriptorFilename);
+				continue;
+			}
+
+			bool bElementEnabled = false;
+			ElementObject->TryGetBoolField(EnabledField, bElementEnabled);
+			if (!bElementEnabled)
+			{
+				UE_LOG(LogGameFeatures, VeryVerbose, TEXT("Skipping adding dependency %s in %s. Plugin is disabled."), *DependencyName, *PluginDescriptorFilename);
+				continue;
+			}
+
+			//Have to get Activate from JSON as it's unique to GFP and not in the PluginManager
+			bool bElementActivate = false;
+			ElementObject->TryGetBoolField(ActivateField, bElementActivate);
+
+			TArray<FString> AssetReferences;
+			const TArray<TSharedPtr<FJsonValue>>* AssetRefsArray = nullptr;
+			ElementObject->TryGetArrayField(AssetReferencesField, AssetRefsArray);
+			if (AssetRefsArray)
+			{
+				for (const TSharedPtr<FJsonValue>& AssetRefElement : *AssetRefsArray)
+				{
+					FString AssetRef;
+					if (AssetRefElement && AssetRefElement->TryGetString(AssetRef))
 					{
-						UE_LOG(LogGameFeatures, Error, TEXT("Error parsing dependency name in %s! Invalid JSON data!"), *PluginDescriptorFilename);
+						AssetReferences.Add(MoveTemp(AssetRef));
 					}
 				}
 			}
+
+			FGameFeaturePluginReferenceDetails& RefDetails = OutPluginDetails.PluginDependencies.Emplace_GetRef();
+			RefDetails.PluginName = MoveTemp(DependencyName);
+			RefDetails.AssetReferences = MoveTemp(AssetReferences);
+			RefDetails.bShouldActivate = bElementActivate;
 		}
 	}
 
@@ -2248,6 +2268,7 @@ bool UGameFeaturesSubsystem::GetGameFeaturePluginDetailsInternal(const FString&
 	{
 		CachedPluginDetailsByFilename.Add(PluginDescriptorFilename, FCachedGameFeaturePluginDetails(OutPluginDetails, FileTimeStamp));
 	}
+
 	return true;
 }
 
@@ -2576,22 +2597,6 @@ TSharedRef<FGameFeaturePluginPredownloadHandle> UGameFeaturesSubsystem::Predownl
 	return Context;
 }
 
-UGameFeaturePluginStateMachine* UGameFeaturesSubsystem::FindGameFeaturePluginStateMachineByPluginName(const FString& PluginName) const
-{
-	for (auto StateMachineIt = GameFeaturePluginStateMachines.CreateConstIterator(); StateMachineIt; ++StateMachineIt)
-	{
-		if (UGameFeaturePluginStateMachine* GFSM = StateMachineIt.Value())
-		{
-			if (GFSM->GetGameFeatureName() == PluginName)
-			{
-				return GFSM;
-			}
-		}
-	}
-
-	return nullptr;
-}
-
 UGameFeaturePluginStateMachine* UGameFeaturesSubsystem::FindGameFeaturePluginStateMachine(const FString& PluginURL) const
 {
 	FGameFeaturePluginIdentifier FindPluginIdentifier(PluginURL);
@@ -2783,6 +2788,23 @@ void UGameFeaturesSubsystem::FinishTermination(UGameFeaturePluginStateMachine* M
 	TerminalGameFeaturePluginStateMachines.RemoveSwap(Machine);
 }
 
+TArray<FString> UGameFeaturesSubsystem::FindPluginAssetDependencies(const FString& PluginDescriptorFilename)
+{
+	FGameFeaturePluginDetails Details;
+	ensure(GetGameFeaturePluginDetailsInternal(PluginDescriptorFilename, Details));
+
+	TArray<FString> OutAssetPaths;
+	for (const FGameFeaturePluginReferenceDetails& RefDetails : Details.PluginDependencies)
+	{
+		for (const FString& AssetRef : RefDetails.AssetReferences)
+		{
+			OutAssetPaths.Emplace(FString::Printf(TEXT("/%s/%s"), *RefDetails.PluginName, *AssetRef));
+		}
+	}
+
+	return OutAssetPaths;
+}
+
 bool UGameFeaturesSubsystem::FindOrCreatePluginDependencyStateMachines(const FString& PluginURL, const FGameFeaturePluginStateMachineProperties& InStateProperties, TArray<UGameFeaturePluginStateMachine*>& OutDependencyMachines)
 {
 	const FString& PluginFilename = InStateProperties.PluginInstalledFilename;
diff --git a/Engine/Plugins/Runtime/GameFeatures/Source/GameFeatures/Public/GameFeatureTypes.h b/Engine/Plugins/Runtime/GameFeatures/Source/GameFeatures/Public/GameFeatureTypes.h
--- a/Engine/Plugins/Runtime/GameFeatures/Source/GameFeatures/Public/GameFeatureTypes.h
+++ b/Engine/Plugins/Runtime/GameFeatures/Source/GameFeatures/Public/GameFeatureTypes.h
@@ -28,6 +28,7 @@ class FString;
 	XSTATE(ErrorWaitingForDependencies,		NSLOCTEXT("GameFeatures", "ErrorWaitingForDependenciesStateDisplayName", "ErrorWaitingForDependencies"))		/* Error state for Installed -> Registered and Registered -> Installed transitions. */ \
 	XSTATE(ErrorRegistering,				NSLOCTEXT("GameFeatures", "ErrorRegisteringDisplayName", "ErrorRegistering"))									/* Error state for Installed -> Registered and Registered -> Installed transitions. */ \
 	XSTATE(WaitingForDependencies,			NSLOCTEXT("GameFeatures", "WaitingForDependenciesStateDisplayName", "WaitingForDependencies"))					/* Transition state Installed -> Registered. In the process of loading code/content for all dependencies into memory. */ \
+	XSTATE(AssetDependencyStreaming,		NSLOCTEXT("GameFeatures", "AssetDependencyStreamingDisplayName", "AssetDependencyStreaming"))					/* Transition state Installed -> Registered. In the process of streaming individual assets from dependencies. */ \
 	XSTATE(Unmounting,						NSLOCTEXT("GameFeatures", "UnmountingStateDisplayName", "Unmounting"))											/* Transition state Registered -> Installed. The content file(s) (i.e. pak file) for the plugin is unmounting. */ \
 	XSTATE(Mounting,						NSLOCTEXT("GameFeatures", "MountingStateDisplayName", "Mounting"))												/* Transition state Installed -> Registered. The content files(s) (i.e. pak file) for the plugin is getting mounted. */ \
 	XSTATE(Unregistering,					NSLOCTEXT("GameFeatures", "UnregisteringStateDisplayName", "Unregistering"))									/* Transition state Registered -> Installed. Cleaning up data gathered in Registering. */ \
diff --git a/Engine/Plugins/Runtime/GameFeatures/Source/GameFeatures/Public/GameFeaturesSubsystem.h b/Engine/Plugins/Runtime/GameFeatures/Source/GameFeatures/Public/GameFeaturesSubsystem.h
--- a/Engine/Plugins/Runtime/GameFeatures/Source/GameFeatures/Public/GameFeaturesSubsystem.h
+++ b/Engine/Plugins/Runtime/GameFeatures/Source/GameFeatures/Public/GameFeaturesSubsystem.h
@@ -233,26 +233,16 @@ void GAMEFEATURES_API LexFromString(EGameFeatureTargetState& Value, const TCHAR*
 struct FGameFeaturePluginReferenceDetails
 {
 	FString PluginName;
-	bool bShouldActivate;
-
-	FGameFeaturePluginReferenceDetails(FString InPluginName, bool bInShouldActivate)
-		: PluginName(MoveTemp(InPluginName))
-		, bShouldActivate(bInShouldActivate)
-	{
-	}
+	TArray<FString> AssetReferences;
+	bool bShouldActivate = false;
 };
 
 struct FGameFeaturePluginDetails
 {
 	TArray<FGameFeaturePluginReferenceDetails> PluginDependencies;
 	TMap<FString, TSharedPtr<class FJsonValue>> AdditionalMetadata;
-	bool bHotfixable;
-	EBuiltInAutoState BuiltInAutoState;
-
-	FGameFeaturePluginDetails()
-		: bHotfixable(false)
-		, BuiltInAutoState(EBuiltInAutoState::Installed)
-	{}
+	bool bHotfixable = false;
+	EBuiltInAutoState BuiltInAutoState = EBuiltInAutoState::Invalid;
 };
 
 struct FBuiltInGameFeaturePluginBehaviorOptions
@@ -713,9 +703,6 @@ private:
 	void PruneCachedGameFeaturePluginDetails(const FString& PluginURL, const FString& PluginDescriptorFilename) const;
 	friend struct FGameFeaturePluginState_Unmounting;
 
-	/** Gets the state machine associated with the specified plugin name */
-	UGameFeaturePluginStateMachine* FindGameFeaturePluginStateMachineByPluginName(const FString& PluginName) const;
-
 	/** Gets the state machine associated with the specified URL */
 	UGameFeaturePluginStateMachine* FindGameFeaturePluginStateMachine(const FString& PluginURL) const;
 
@@ -742,6 +729,9 @@ private:
 	void FinishTermination(UGameFeaturePluginStateMachine* Machine);
 	friend class UGameFeaturePluginStateMachine;
 
+	TArray<FString> FindPluginAssetDependencies(const FString& PluginDescriptorFilename);
+	friend struct FGameFeaturePluginState_AssetDependencyStreaming;
+
 	/** Handler for when a state machine requests its dependencies. Returns false if the dependencies could not be read */
 	bool FindOrCreatePluginDependencyStateMachines(const FString& PluginURL, const FGameFeaturePluginStateMachineProperties& InStateProperties, TArray<UGameFeaturePluginStateMachine*>& OutDependencyMachines);
 	template <typename> friend struct FTransitionDependenciesGameFeaturePluginState;
diff --git a/Engine/Plugins/Runtime/GameFeatures/Source/GameFeaturesEditor/Private/SGameFeatureStateWidget.cpp b/Engine/Plugins/Runtime/GameFeatures/Source/GameFeaturesEditor/Private/SGameFeatureStateWidget.cpp
--- a/Engine/Plugins/Runtime/GameFeatures/Source/GameFeaturesEditor/Private/SGameFeatureStateWidget.cpp
+++ b/Engine/Plugins/Runtime/GameFeatures/Source/GameFeaturesEditor/Private/SGameFeatureStateWidget.cpp
@@ -65,7 +65,7 @@ FText SGameFeatureStateWidget::GetDisplayNameOfState(EGameFeaturePluginState Sta
 
 FText SGameFeatureStateWidget::GetTooltipOfState(EGameFeaturePluginState State)
 {
-	static_assert((int32)EGameFeaturePluginState::MAX == 34, "");
+	static_assert((int32)EGameFeaturePluginState::MAX == 35, "");
 
 	switch (State)
 	{
@@ -105,6 +105,8 @@ FText SGameFeatureStateWidget::GetTooltipOfState(EGameFeaturePluginState State)
 		return LOCTEXT("StateTooltip_ErrorRegistering", "Error state for Installed -> Registered and Registered -> Installed transitions.");
 	case EGameFeaturePluginState::WaitingForDependencies:
 		return LOCTEXT("StateTooltip_WaitingForDependencies", "Transition state Installed -> Registered. In the process of loading code/content for all dependencies into memory.");
+	case EGameFeaturePluginState::AssetDependencyStreaming:
+		return LOCTEXT("StateTooltip_AssetDependencyStreaming", "Transition state Installed -> Registered. In the process of streaming individual assets from dependencies.");
 	case EGameFeaturePluginState::Unmounting:
 		return LOCTEXT("StateTooltip_Unmounting", "Transition state Registered -> Installed. The content file(s) (i.e. pak file) for the plugin is unmounting.");
 	case EGameFeaturePluginState::Mounting:
