Add clarification to GraphConvert comments for usage
Fix issue in AddChunkOverride where not all primary assets wouuld be found when performing override. Don't add override until assets scanning finishes.

Add new function - BuildChunkMap
This allows users to build the chunk map themselves rather than having all assets tested one by one to determine their chunk.

UGameFeatureAction_AddChunkOverride::GetChunkForPackage is now mostly unecessary except when wanting to still pull startup packages like the GameFeatureData into an explicit chunk.
Users are expected to override GetPackageChunkId if they want the GFD to not be placed in the startup chunk.

diff --git a/Engine/Plugins/Runtime/GameFeatures/Source/GameFeatures/Private/GameFeatureAction_AddChunkOverride.cpp b/Engine/Plugins/Runtime/GameFeatures/Source/GameFeatures/Private/GameFeatureAction_AddChunkOverride.cpp
--- a/Engine/Plugins/Runtime/GameFeatures/Source/GameFeatures/Private/GameFeatureAction_AddChunkOverride.cpp
+++ b/Engine/Plugins/Runtime/GameFeatures/Source/GameFeatures/Private/GameFeatureAction_AddChunkOverride.cpp
@@ -5,6 +5,11 @@
 #include "Engine/AssetManagerSettings.h"
 #include "GameFeatureData.h"
 #include "Misc/MessageDialog.h"
+#include "Misc/PathViews.h"
+
+#if WITH_EDITOR
+#include "Commandlets/ChunkDependencyInfo.h"
+#endif
 
 #include UE_INLINE_GENERATED_CPP_BY_NAME(GameFeatureAction_AddChunkOverride)
 
@@ -17,7 +22,7 @@ DEFINE_LOG_CATEGORY_STATIC(LogAddChunkOverride, Log, All);
 
 namespace GameFeatureAction_AddChunkOverride
 {
-	static TMap<int32, FString> ChunkIdToPluginMap;
+	static TMap<int32, TArray<FString>> ChunkIdToPluginMap;
 	static TMap<FString, int32> PluginToChunkId;
 }
 
@@ -28,7 +33,8 @@ void UGameFeatureAction_AddChunkOverride::OnGameFeatureRegistering()
 	const bool bShouldAddChunkOverride = ShouldAddChunkOverride.IsBound() ? ShouldAddChunkOverride.Execute(GetTypedOuter<UGameFeatureData>()) : true;
 	if (bShouldAddChunkOverride)
 	{
-		AddChunkIdOverride();
+		TWeakObjectPtr<UGameFeatureAction_AddChunkOverride> WeakThis(this);
+		UAssetManager::CallOrRegister_OnCompletedInitialScan(FSimpleMulticastDelegate::FDelegate::CreateUObject(this, &UGameFeatureAction_AddChunkOverride::AddChunkIdOverride));
 	}
 }
 
@@ -38,105 +44,36 @@ void UGameFeatureAction_AddChunkOverride::OnGameFeatureUnregistering()
 }
 
 #if WITH_EDITOR
-void UGameFeatureAction_AddChunkOverride::GetChunkForPackage(const FName PackageName, const int32 DefaultGameChunk, TArray<int32>& OutChunkList)
-{
-	TSet<FPrimaryAssetId> Managers;
-	UAssetManager::Get().GetPackageManagers(PackageName, true, Managers);
-	GetChunkForPackage(PackageName.ToString(), Managers, DefaultGameChunk, OutChunkList);
-}
-
-void UGameFeatureAction_AddChunkOverride::GetChunkForPackage(const FString& PackageName, const TSet<FPrimaryAssetId>& Managers, const int32 DefaultGameChunk, TArray<int32>& OutChunkList)
+TOptional<int32> UGameFeatureAction_AddChunkOverride::GetChunkForPackage(const FString& PackageName)
 {
 	if (GameFeatureAction_AddChunkOverride::PluginToChunkId.Num() == 0)
 	{
-		return;
+		return TOptional<int32>();
 	}
 
-	auto ResolveMultiChunkDependecies = [PackageName, &Managers, DefaultGameChunk,  &OutChunkList]()
-	{
-		UE_LOG(LogAddChunkOverride, VeryVerbose, TEXT("%s was referenced by one than one chunk"), *PackageName);
-		for (const int32 OutChunkId : OutChunkList)
-		{
-			UE_LOG(LogAddChunkOverride, VeryVerbose, TEXT("%s was referenced by chunk %d"), *PackageName, OutChunkId);
-		}
-
-		if (OutChunkList.Num() > 1 && !OutChunkList.Contains(0))
-		{
-			UE_LOG(LogAddChunkOverride, VeryVerbose, TEXT("Forcing %s into gameplay chunk %d. It was referend by multiple GFPs which might load at different times."), *PackageName, DefaultGameChunk);
-			OutChunkList.Reset();
-			OutChunkList.Add(DefaultGameChunk);
-		}
-		else
-		{
-			// If multiple package mangers exist for this package with different chunk IDs it should go into the default game chunk.
-			UAssetManager& AssetManager = UAssetManager::Get();
-			TSet<int32> ManagerChunkIds;
-			for (const FPrimaryAssetId& PrimaryAssetId : Managers)
-			{
-				FPrimaryAssetRules Rules = AssetManager.GetPrimaryAssetRules(PrimaryAssetId);
-				if (Rules.ChunkId != INDEX_NONE)
-				{
-					ManagerChunkIds.Add(Rules.ChunkId);
-				}
-			}
-			if (ManagerChunkIds.Num() > 1)
-			{
-				UE_LOG(LogAddChunkOverride, Log, TEXT("Forcing %s into gameplay chunk %d. It was referend by multiple GFPs which might load at different times. Package managers with a valid chunkID might not have been registered for this type."), *PackageName, DefaultGameChunk);
-				OutChunkList.Reset();
-				OutChunkList.Add(DefaultGameChunk);
-			}
-		}
-	};
-
 	static const FString EngineDir(TEXT("/Engine/"));
 	static const FString GameDir(TEXT("/Game/"));
 	if (PackageName.StartsWith(EngineDir, ESearchCase::CaseSensitive))
 	{
-		return;
+		return TOptional<int32>();
 	}
 	else if (PackageName.StartsWith(GameDir, ESearchCase::CaseSensitive))
 	{
-		ResolveMultiChunkDependecies();
+		return TOptional<int32>();
 	}
 	else
 	{
-		TArray<FString> BrokenString;
-		PackageName.ParseIntoArray(BrokenString, TEXT("/"));
-		if (BrokenString.Num() > 0)
+		FString MountPointName = FString(FPathViews::GetMountPointNameFromPath(PackageName));
+		if (GameFeatureAction_AddChunkOverride::PluginToChunkId.Contains(MountPointName))
 		{
-			FString PluginName = BrokenString[0];
-			if (GameFeatureAction_AddChunkOverride::PluginToChunkId.Contains(PluginName))
-			{
-				int32 ExpectedChunkId = GameFeatureAction_AddChunkOverride::PluginToChunkId[PluginName];
-				if (OutChunkList.Contains(ExpectedChunkId) && OutChunkList.Num() > 1)
-				{
-					UE_LOG(LogAddChunkOverride, VeryVerbose, TEXT("%s is referenced by expected chunk %d but is also referend by other chunks."), *PackageName, ExpectedChunkId);
-					if (LogAddChunkOverride.GetVerbosity() == ELogVerbosity::VeryVerbose)
-					{
-						for (const int32 ChunkId : OutChunkList)
-						{
-							UE_LOG(LogAddChunkOverride, VeryVerbose, TEXT("%s was referenced by chunk %d"), *PackageName, ChunkId);
-						}
-					}
-				}
-				else
-				{
-					UE_LOG(LogAddChunkOverride, VeryVerbose, TEXT("%s was expected to be in chunk %d but not found there. This could be because of a GameplayCue or this packages is only cooking because it is referend by another plugin"), *PackageName, ExpectedChunkId);
-				}
-
-				UE_LOG(LogAddChunkOverride, VeryVerbose, TEXT("Forcing %s into chunk %d"), *PackageName, ExpectedChunkId);
-				OutChunkList.Reset();
-				OutChunkList.Add(ExpectedChunkId);
-			}
-			else if (OutChunkList.Num() > 1)
-			{
-				ResolveMultiChunkDependecies();
-			}
+			const int32 ExpectedChunkId = GameFeatureAction_AddChunkOverride::PluginToChunkId[MountPointName];
+			return TOptional<int32>(ExpectedChunkId);
 		}
 	}
+	return TOptional<int32>();
 }
 
-FString UGameFeatureAction_AddChunkOverride::GetPluginNameFromChunkID(int32 ChunkID)
+TArray<FString> UGameFeatureAction_AddChunkOverride::GetPluginNameFromChunkID(int32 ChunkID)
 {
 	return GameFeatureAction_AddChunkOverride::ChunkIdToPluginMap.FindRef(ChunkID);
 }
@@ -207,23 +144,33 @@ void UGameFeatureAction_AddChunkOverride::AddChunkIdOverride()
 		UE_LOG(LogAddChunkOverride, Error, TEXT("ChunkId is negative. Unable to override to a negative chunk"));
 		return;
 	}
-	if (GameFeatureAction_AddChunkOverride::ChunkIdToPluginMap.Contains(ChunkId))
+
+	if (UGameFeatureData* GameFeatureData = GetTypedOuter<UGameFeatureData>())
 	{
-		FString PluginName;
-		if (UGameFeatureData* GameFeatureData = GetTypedOuter<UGameFeatureData>())
+		UChunkDependencyInfo* DependencyInfo = GetMutableDefault<UChunkDependencyInfo>();
+		
+
+		if (FChunkDependency* ExistingDep = DependencyInfo->DependencyArray.FindByPredicate([CheckChunk = ChunkId](const FChunkDependency& ChunkDep)
+			{
+				return ChunkDep.ChunkID == CheckChunk;
+			}))
 		{
-			GameFeatureData->GetPluginName(PluginName);
+			// If we found this chunk it might have been auto generated. Update this instead of adding ours.
+			ExistingDep->ParentChunkID = ParentChunk;
 		}
-		UE_LOG(LogAddChunkOverride, Error, TEXT("ChunkId (%d) is already in use by %s. Manually resolve the conflict for %s"), ChunkId, *GameFeatureAction_AddChunkOverride::ChunkIdToPluginMap[ChunkId], *PluginName);
-		FMessageDialog::Open(EAppMsgType::Ok, FText::Format(LOCTEXT("AddChunkOverride_IdConflight", "Chunk Id is already in use by '{0}'."), FText::FromString(PluginName)));
-		return;
-	}
+		else
+		{
+			FChunkDependency NewChunkDependency;
+			NewChunkDependency.ChunkID = ChunkId;
+			NewChunkDependency.ParentChunkID = ParentChunk;
+			DependencyInfo->DependencyArray.Add(NewChunkDependency);
+		}
+		DependencyInfo->GetOrBuildChunkDependencyGraph(ChunkId, true);
 
-	if (UGameFeatureData* GameFeatureData = GetTypedOuter<UGameFeatureData>())
-	{
 		FString PluginName;
 		GameFeatureData->GetPluginName(PluginName);
-		GameFeatureAction_AddChunkOverride::ChunkIdToPluginMap.Add(ChunkId, PluginName);
+		TArray<FString>& PluginsInChunk = GameFeatureAction_AddChunkOverride::ChunkIdToPluginMap.FindOrAdd(ChunkId);
+		PluginsInChunk.Add(PluginName);
 		GameFeatureAction_AddChunkOverride::PluginToChunkId.Add(PluginName, ChunkId);
 		LastChunkIdUsed = ChunkId;
 		UE_LOG(LogAddChunkOverride, Log, TEXT("Plugin(%s) will cook assets into chunk(%d)"), *PluginName, ChunkId);
@@ -261,7 +208,14 @@ void UGameFeatureAction_AddChunkOverride::RemoveChunkIdOverride()
 		FString PluginName;
 		GameFeatureData->GetPluginName(PluginName);
 
-		ensure(GameFeatureAction_AddChunkOverride::ChunkIdToPluginMap.Remove(LastChunkIdUsed));
+		if (GameFeatureAction_AddChunkOverride::ChunkIdToPluginMap.Contains(LastChunkIdUsed))
+		{
+			GameFeatureAction_AddChunkOverride::ChunkIdToPluginMap[LastChunkIdUsed].Remove(PluginName);
+			if (GameFeatureAction_AddChunkOverride::ChunkIdToPluginMap[LastChunkIdUsed].IsEmpty())
+			{
+				GameFeatureAction_AddChunkOverride::ChunkIdToPluginMap.Remove(LastChunkIdUsed);
+			}
+		}
 		ensure(GameFeatureAction_AddChunkOverride::PluginToChunkId.Remove(PluginName));
 		UE_LOG(LogAddChunkOverride, Log, TEXT("Removing ChunkId override (%d) for Plugin (%s)"), LastChunkIdUsed, *PluginName);
 		LastChunkIdUsed = -1;
@@ -308,7 +262,7 @@ int32 UGameFeatureAction_AddChunkOverride::GenerateUniqueChunkId() const
 	}
 	else if (GameFeatureAction_AddChunkOverride::ChunkIdToPluginMap.Contains(NewChunkId))
 	{
-		UE_LOG(LogAddChunkOverride, Warning, TEXT("ChunkId(%d) is in use by %s. Unable to autogenerate unique id. Lowest allowed ChunkId(%d)"), NewChunkId, *GameFeatureAction_AddChunkOverride::ChunkIdToPluginMap[NewChunkId], LowestAllowedChunkIndexForAutoGeneration);
+		UE_LOG(LogAddChunkOverride, Warning, TEXT("ChunkId(%d) is in use by %s. Unable to autogenerate unique id. Lowest allowed ChunkId(%d)"), NewChunkId, *FString::Join(GameFeatureAction_AddChunkOverride::ChunkIdToPluginMap[ChunkId], TEXT(",")), LowestAllowedChunkIndexForAutoGeneration);
 		FMessageDialog::Open(EAppMsgType::Ok, LOCTEXT("AddChunkOverride_UsedChunkId", "Unable to auto generate unique valid Chunk Id. Please manually assign a valid Chunk Id"));
 		NewChunkId = -1;
 	}
diff --git a/Engine/Plugins/Runtime/GameFeatures/Source/GameFeatures/Public/GameFeatureAction_AddChunkOverride.h b/Engine/Plugins/Runtime/GameFeatures/Source/GameFeatures/Public/GameFeatureAction_AddChunkOverride.h
--- a/Engine/Plugins/Runtime/GameFeatures/Source/GameFeatures/Public/GameFeatureAction_AddChunkOverride.h
+++ b/Engine/Plugins/Runtime/GameFeatures/Source/GameFeatures/Public/GameFeatureAction_AddChunkOverride.h
@@ -37,16 +37,14 @@ public:
 #if WITH_EDITOR
 	/**
 	 * Given the package name will check if this is a package from a GFP that we want to assign a specific chunk to.
-	 * Will modify the given OutChunkList removing entries so that the package prefers to go into the specified override chunk.
-	 * When multiple GFPs reference a /Game package the DefaultGameChunk will be preferred.
-	 * /Engine packages are ignored by the function
+	 * returns the override chunk for this package if one is set.
 	 * 
+	 * Should be used in combination with overriding UAssetManager::GetPackageChunkIds so that you are able to reassign a startup package.
 	 * This can be necessary to reassign startup packages such as the GameFeatureData asset.
 	 */
-	GAMEFEATURES_API static void GetChunkForPackage(const FName PackageName, const int32 DefaultGameChunk, TArray<int32>& OutChunkList);
-	GAMEFEATURES_API static void GetChunkForPackage(const FString& PackageName, const TSet<FPrimaryAssetId>& Managers, const int32 DefaultGameChunk, TArray<int32>& OutChunkList);
+	GAMEFEATURES_API static TOptional<int32> GetChunkForPackage(const FString& PackageName);
 
-	GAMEFEATURES_API static FString GetPluginNameFromChunkID(int32 ChunkID);
+	GAMEFEATURES_API static TArray<FString> GetPluginNameFromChunkID(int32 ChunkID);
 
 	/** UObject overrides */
 	virtual void PostRename(UObject* OldOuter, const FName OldName) override;
@@ -74,6 +72,13 @@ private:
 	UPROPERTY(EditAnywhere, Category = "Asset Management", meta=(EditCondition="bShouldOverrideChunk"))
 	int32 ChunkId = -1;
 
+	/**
+	 * What Chunk we are parented to.
+	 * This is used by the ChunkDependencyInfo for when mutiple chunk overrides might conflict requiring assets to be pulled into a lower chunk
+	 */
+	UPROPERTY(EditAnywhere, Category = "Asset Management", meta=(EditCondition="bShouldOverrideChunk"))
+	int32 ParentChunk = 10;
+
 	/**
 	 * Internal tracking variable to remove the chunkId used when the ChunkId property changes.
 	 */
diff --git a/Engine/Source/Editor/UnrealEd/Private/Commandlets/AssetRegistryGenerator.cpp b/Engine/Source/Editor/UnrealEd/Private/Commandlets/AssetRegistryGenerator.cpp
--- a/Engine/Source/Editor/UnrealEd/Private/Commandlets/AssetRegistryGenerator.cpp
+++ b/Engine/Source/Editor/UnrealEd/Private/Commandlets/AssetRegistryGenerator.cpp
@@ -1751,6 +1751,8 @@ bool FAssetRegistryGenerator::SaveAssetRegistry(const FString& SandboxPath, bool
 		// an asset registry for each chunk.
 		if (FinalChunkManifests.Num() > 0)
 		{
+			// before final assignment ensure the database has been made.
+			UAssetManager::Get().UpdateManagementDatabase();
 			// Pass over all chunks and build a mapping of chunk index to asset registry name. All chunks that don't have a unique registry are assigned to the "generic bucket"
 			// which will be written to the master asset registry in chunk 0
 			for (int32 PakchunkIndex = 0; PakchunkIndex < FinalChunkManifests.Num(); ++PakchunkIndex)
@@ -2566,11 +2568,7 @@ void FAssetRegistryGenerator::GetChunkAssignments(TArray<TSet<FName>>& OutAssign
 	{
 		// chunk 0 is special as it also contains startup packages
 		TSet<FName> PackagesInChunk0;
-
-		for(const FName& Package : StartupPackages)
-		{
-			PackagesInChunk0.Add(Package);
-		}
+		PackagesInChunk0.Append(StartupPackages);
 
 		if (ChunkManifests[0])
 		{
diff --git a/Engine/Source/Editor/UnrealEd/Private/Commandlets/ChunkDependencyInfo.cpp b/Engine/Source/Editor/UnrealEd/Private/Commandlets/ChunkDependencyInfo.cpp
--- a/Engine/Source/Editor/UnrealEd/Private/Commandlets/ChunkDependencyInfo.cpp
+++ b/Engine/Source/Editor/UnrealEd/Private/Commandlets/ChunkDependencyInfo.cpp
@@ -1,6 +1,10 @@
 // Copyright Epic Games, Inc. All Rights Reserved.
 
 #include "Commandlets/ChunkDependencyInfo.h"
+#include "Algo/Unique.h"
+
+DEFINE_LOG_CATEGORY_STATIC(LogChunkDependencyInfo, Log, All);
+
 
 UChunkDependencyInfo::UChunkDependencyInfo(const FObjectInitializer& ObjectInitializer)
 	: Super(ObjectInitializer)
@@ -8,12 +12,16 @@ UChunkDependencyInfo::UChunkDependencyInfo(const FObjectInitializer& ObjectIniti
 	CachedHighestChunk = -1;
 }
 
-const FChunkDependencyTreeNode* UChunkDependencyInfo::GetOrBuildChunkDependencyGraph(int32 HighestChunk)
+const FChunkDependencyTreeNode* UChunkDependencyInfo::GetOrBuildChunkDependencyGraph(int32 HighestChunk, bool bForceRebuild)
 {
 	if (HighestChunk > CachedHighestChunk)
 	{
 		return BuildChunkDependencyGraph(HighestChunk);
 	}
+	else if (bForceRebuild)
+	{
+		return BuildChunkDependencyGraph(CachedHighestChunk);
+	}
 	return &RootTreeNode;
 }
 
@@ -114,4 +122,57 @@ void UChunkDependencyInfo::RemoveRedundantChunks(TArray<int32>& ChunkIDs) const
 			}
 		}
 	}
-}
\ No newline at end of file
+}
+
+const FChunkDependencyTreeNode* LowestCommonAncestor(const FChunkDependencyTreeNode* RootNode, TSet<int32>& FoundSet, const TSet<int32>& FindSet)
+{
+	if (RootNode)
+	{
+		for (const FChunkDependencyTreeNode& ChildNode : RootNode->ChildNodes)
+		{
+			if (const FChunkDependencyTreeNode* ReturnNode = LowestCommonAncestor(&ChildNode, FoundSet, FindSet))
+			{
+				return ReturnNode;
+			}
+		}
+
+		FoundSet.Add(RootNode->ChunkID);
+
+		for (int32 Item : FindSet)
+		{
+			if (!FoundSet.Contains(Item))
+			{
+				return nullptr;
+			}
+		}
+
+		return RootNode;
+	}
+	return nullptr;
+}
+
+
+int32 UChunkDependencyInfo::FindHighestSharedChunk(const TArray<int32>& ChunkIDs) const
+{
+	if (ChunkIDs.Num() == 0)
+	{
+		return 0;
+	}
+	if (ChunkIDs.Num() == 1)
+	{
+		return ChunkIDs[0];
+	}
+
+	TSet<int32> FoundSet;
+	TSet<int32> FindSet;
+	FindSet.Append(ChunkIDs);
+	if (const FChunkDependencyTreeNode* BestParent = LowestCommonAncestor(&RootTreeNode, FoundSet, FindSet))
+	{
+		return BestParent->ChunkID;
+	}
+	else
+	{
+		UE_LOG(LogChunkDependencyInfo, Error, TEXT("Unable to find parent."));
+		return 0;
+	}
+}
diff --git a/Engine/Source/Editor/UnrealEd/Public/Commandlets/ChunkDependencyInfo.h b/Engine/Source/Editor/UnrealEd/Public/Commandlets/ChunkDependencyInfo.h
--- a/Engine/Source/Editor/UnrealEd/Public/Commandlets/ChunkDependencyInfo.h
+++ b/Engine/Source/Editor/UnrealEd/Public/Commandlets/ChunkDependencyInfo.h
@@ -53,7 +53,7 @@ class UChunkDependencyInfo : public UObject
 	GENERATED_UCLASS_BODY()
 
 	/** Will return an existing dependency graph, or build if it necessary */
-	UNREALED_API const FChunkDependencyTreeNode* GetOrBuildChunkDependencyGraph(int32 HighestChunk = 0);
+	UNREALED_API const FChunkDependencyTreeNode* GetOrBuildChunkDependencyGraph(int32 HighestChunk = 0, bool bForceRebuild = false);
 
 	/** Will create a dependency tree starting with RootTreeNode. If HighestChunk is == 0 it will only add dependencies on 0 for chunks already in the dependencies list */
 	UNREALED_API const FChunkDependencyTreeNode* BuildChunkDependencyGraph(int32 HighestChunk);
@@ -61,6 +61,11 @@ class UChunkDependencyInfo : public UObject
 	/** Removes redundant chunks from a chunk list */
 	UNREALED_API void RemoveRedundantChunks(TArray<int32>& ChunkIDs) const;
 
+	/** Given a set of chunk Ids will use the dependencies to find which chunk is shared among all the given chunks.
+	All chunks should converge to the 0 startup chunk.
+	Potentially will return a chunk that is not in the given input array */
+	UNREALED_API int32 FindHighestSharedChunk(const TArray<int32>& ChunkIDs) const;
+
 	/** List of dependencies used to remove redundant chunks */
 	UPROPERTY(config)
 	TArray<FChunkDependency> DependencyArray;
diff --git a/Engine/Source/Runtime/Core/Public/Experimental/Graph/GraphConvert.h b/Engine/Source/Runtime/Core/Public/Experimental/Graph/GraphConvert.h
--- a/Engine/Source/Runtime/Core/Public/Experimental/Graph/GraphConvert.h
+++ b/Engine/Source/Runtime/Core/Public/Experimental/Graph/GraphConvert.h
@@ -297,6 +297,8 @@ CORE_API FGraph ConstructTransposeGraph(
  * Return a new graph where cycles in the input graph have been replaced by a single vertex. The new graph is 
  * topologically sorted from root to leaf. If the input graph has no cycles, the function returns false and the 
  * OutGraph is reset to empty.
+ * 
+ * After conversion you can iterate over the OutCondensationVertexToInputVertex mapping to go through the map topologically.
  *
  * @param Graph The input graph, may contain cycles. @see ConvertToGraph for a description of the format.
  * @param OutGraph Output value that holds the graph. 
diff --git a/Engine/Source/Runtime/Engine/Classes/Engine/AssetManager.h b/Engine/Source/Runtime/Engine/Classes/Engine/AssetManager.h
--- a/Engine/Source/Runtime/Engine/Classes/Engine/AssetManager.h
+++ b/Engine/Source/Runtime/Engine/Classes/Engine/AssetManager.h
@@ -741,6 +741,15 @@ protected:
 	/** Called after PIE ends, resets loading state */
 	ENGINE_API virtual void EndPIE(bool bStartSimulate);
 
+	/**
+	 * Used by UpdateManagementDatabase to build the CachedChunkMap.
+	 * Should return a new cachedChunkMap.
+	 * 
+	 * @param PackagesToUpdateChunksFor - the packages needing chunk assignment as part of the management database update.
+	 * @return a new CachedChunkMap to be used.
+	 */
+	ENGINE_API virtual TMap<int32, FAssetManagerChunkInfo> BuildChunkMap(const TSet<FName>& PackagesToUpdateChunksFor) const;
+
 	/** Copy of the asset state before PIE was entered, return to that when PIE completes */
 	TMap<FPrimaryAssetId, TArray<FName>> PrimaryAssetStateBeforePIE;
 
diff --git a/Engine/Source/Runtime/Engine/Private/AssetManager.cpp b/Engine/Source/Runtime/Engine/Private/AssetManager.cpp
--- a/Engine/Source/Runtime/Engine/Private/AssetManager.cpp
+++ b/Engine/Source/Runtime/Engine/Private/AssetManager.cpp
@@ -3526,14 +3526,6 @@ void UAssetManager::PostInitialAssetScan()
 
 	bIsPrimaryAssetDirectoryCurrent = true;
 
-#if WITH_EDITOR
-	if (bUpdateManagementDatabaseAfterScan)
-	{
-		bUpdateManagementDatabaseAfterScan = false;
-		UpdateManagementDatabase(true);
-	}
-#endif
-
 	if (!bHasCompletedInitialScan)
 	{
 		// Done with initial scan, fire delegate exactly once. This does not happen on editor refreshes
@@ -3541,6 +3533,14 @@ void UAssetManager::PostInitialAssetScan()
 		OnCompletedInitialScanDelegate.Broadcast();
 		OnCompletedInitialScanDelegate.Clear();
 	}
+
+#if WITH_EDITOR
+	if (bUpdateManagementDatabaseAfterScan)
+	{
+		bUpdateManagementDatabaseAfterScan = false;
+		UpdateManagementDatabase(true);
+	}
+#endif
 }
 
 bool UAssetManager::GetManagedPackageList(FPrimaryAssetId PrimaryAssetId, TArray<FName>& PackagePathList) const
@@ -4048,13 +4048,10 @@ void UAssetManager::UpdateManagementDatabase(bool bForceRefresh)
 		AssetRegistry.SetManageReferences(NoReferenceManagementMap, false, UE::AssetRegistry::EDependencyCategory::None, ExistingManagedNodes);
 	}
 
-
 	TMultiMap<FAssetIdentifier, FAssetIdentifier> PrimaryAssetIdManagementMap;
 	TArray<int32> ChunkList;
 	TArray<int32> ExistingChunkList;
 
-	CachedChunkMap.Empty(); // Remove previous entries before we start adding to it
-
 	// Update management parent list, which is PrimaryAssetId -> PrimaryAssetId
 	for (const TPair<FName, TSharedRef<FPrimaryAssetTypeData>>& TypePair : AssetTypeMap)
 	{
@@ -4099,8 +4096,6 @@ void UAssetManager::UpdateManagementDatabase(bool bForceRefresh)
 			for (int32 ChunkId : ChunkList)
 			{
 				FPrimaryAssetId ChunkPrimaryAsset = CreatePrimaryAssetIdFromChunkId(ChunkId);
-
-				CachedChunkMap.FindOrAdd(ChunkId).ExplicitAssets.Add(PrimaryAssetId);
 				PrimaryAssetIdManagementMap.Add(ChunkPrimaryAsset, PrimaryAssetId);
 			}
 		}
@@ -4114,29 +4109,7 @@ void UAssetManager::UpdateManagementDatabase(bool bForceRefresh)
 	UProjectPackagingSettings* ProjectPackagingSettings = GetMutableDefault<UProjectPackagingSettings>();
 	if (ProjectPackagingSettings && ProjectPackagingSettings->bGenerateChunks)
 	{
-		// Update the editor preview chunk package list for all chunks, but only if we actually care about chunks
-		// bGenerateChunks is settable per platform, but should be enabled on the default platform for preview to work
-		TArray<int32> OverrideChunkList;
-		for (FName PackageName : PackagesToUpdateChunksFor)
-		{
-			ChunkList.Reset();
-			OverrideChunkList.Reset();
-			GetPackageChunkIds(PackageName, nullptr, ExistingChunkList, ChunkList, &OverrideChunkList);
-
-			if (ChunkList.Num() > 0)
-			{
-				for (int32 ChunkId : ChunkList)
-				{
-					CachedChunkMap.FindOrAdd(ChunkId).AllAssets.Add(PackageName);
-
-					if (OverrideChunkList.Contains(ChunkId))
-					{
-						// This was in the override list, so add an explicit dependency
-						CachedChunkMap.FindOrAdd(ChunkId).ExplicitAssets.Add(PackageName);
-					}
-				}
-			}
-		}
+		CachedChunkMap = BuildChunkMap(PackagesToUpdateChunksFor);
 	}
 
 	bIsManagementDatabaseCurrent = true;
@@ -4680,6 +4653,74 @@ void UAssetManager::EndPIE(bool bStartSimulate)
 	}
 }
 
+TMap<int32, FAssetManagerChunkInfo> UAssetManager::BuildChunkMap(const TSet<FName>& PackagesToUpdateChunksFor) const
+{
+	TMap<int32, FAssetManagerChunkInfo> ReturnChunkMap;
+
+	TArray<int32> ChunkList;
+	TArray<int32> ExistingChunkList;
+
+	// Update management parent list, which is PrimaryAssetId -> PrimaryAssetId
+	for (const TPair<FName, TSharedRef<FPrimaryAssetTypeData>>& TypePair : AssetTypeMap)
+	{
+		const FPrimaryAssetTypeData& TypeData = TypePair.Value.Get();
+		for (const TPair<FName, FPrimaryAssetData>& NamePair : TypeData.GetAssets())
+		{
+			const FPrimaryAssetData& NameData = NamePair.Value;
+			FPrimaryAssetId PrimaryAssetId(TypePair.Key, NamePair.Key);
+			const FSoftObjectPath& AssetRef = NameData.GetAssetPtr().ToSoftObjectPath();
+
+			TSet<FPrimaryAssetId> Managers;
+			if (AssetRef.IsValid())
+			{
+				FName PackageName = FName(*AssetRef.GetLongPackageName());
+				GetPackageManagers(PackageName, false, Managers);
+			}
+			else
+			{
+				Managers.Add(PrimaryAssetId);
+			}
+
+			// Compute chunk assignment and store those as manager references
+			ChunkList.Reset();
+			GetPrimaryAssetSetChunkIds(Managers, nullptr, ExistingChunkList, ChunkList);
+
+			for (int32 ChunkId : ChunkList)
+			{
+				FPrimaryAssetId ChunkPrimaryAsset = CreatePrimaryAssetIdFromChunkId(ChunkId);
+
+				ReturnChunkMap.FindOrAdd(ChunkId).ExplicitAssets.Add(PrimaryAssetId);
+			}
+		}
+	}
+
+	// Update the editor preview chunk package list for all chunks, but only if we actually care about chunks
+	// bGenerateChunks is settable per platform, but should be enabled on the default platform for preview to work
+	TArray<int32> OverrideChunkList;
+	for (FName PackageName : PackagesToUpdateChunksFor)
+	{
+		ChunkList.Reset();
+		OverrideChunkList.Reset();
+		GetPackageChunkIds(PackageName, nullptr, ExistingChunkList, ChunkList, &OverrideChunkList);
+
+		if (ChunkList.Num() > 0)
+		{
+			for (int32 ChunkId : ChunkList)
+			{
+				ReturnChunkMap.FindOrAdd(ChunkId).AllAssets.Add(PackageName);
+
+				if (OverrideChunkList.Contains(ChunkId))
+				{
+					// This was in the override list, so add an explicit dependency
+					ReturnChunkMap.FindOrAdd(ChunkId).ExplicitAssets.Add(PackageName);
+				}
+			}
+		}
+	}
+
+	return ReturnChunkMap;
+}
+
 void UAssetManager::ReinitializeFromConfig()
 {
 	// We specifically do not reset AssetRuleOverrides as those can be set by something other than inis
