GFP localization loading/unloading tweaks

Added the "GameFeaturePlugin.DeferLocalizationDataLoad" CVar to control whether GFP localization loading has the old (5.3) or new (5.4) behavior:
  * When true, localization data will load on GFP 'loading' and unload on GFP 'unloading' (new behavior).
  * When false, localization data will load on GFP 'mounting' and unload on GFP 'unmounting' (old behavior).

This also adjusts the new behavior to be on loading/unloading rather than activating/deactivating, as this makes more sense, and allows the localization data more time to load async prior to the wait at the end of 'activating'.

diff --git a/Engine/Plugins/Runtime/GameFeatures/Source/GameFeatures/Private/GameFeaturePluginStateMachine.cpp b/Engine/Plugins/Runtime/GameFeatures/Source/GameFeatures/Private/GameFeaturePluginStateMachine.cpp
--- a/Engine/Plugins/Runtime/GameFeatures/Source/GameFeatures/Private/GameFeaturePluginStateMachine.cpp
+++ b/Engine/Plugins/Runtime/GameFeatures/Source/GameFeatures/Private/GameFeaturePluginStateMachine.cpp
@@ -62,6 +62,11 @@ namespace UE::GameFeatures
 		TEXT("Comma-separated list of names of plugins for which to skip verification."),
 		ECVF_Default);
 
+	static bool bDeferLocalizationDataLoad = true;
+	static FAutoConsoleVariableRef CVarDeferLocalizationLoading(TEXT("GameFeaturePlugin.DeferLocalizationDataLoad"),
+		bDeferLocalizationDataLoad,
+		TEXT("True if we should defer loading the localization data until 'loading' (new behavior), or false to load it on 'mounting' (old behavior)."));
+
 	static TAutoConsoleVariable<bool> CVarAsyncLoad(TEXT("GameFeaturePlugin.AsyncLoad"),
 		true,
 		TEXT("Enable to use aysnc loading"));
@@ -1617,6 +1622,11 @@ struct FGameFeaturePluginState_Unmounting : public FGameFeaturePluginState
 		if (TSharedPtr<IPlugin> Plugin = IPluginManager::Get().FindPlugin(StateProperties.PluginName);
 			Plugin && Plugin->GetDescriptor().bExplicitlyLoaded)
 		{
+			if (!UE::GameFeatures::bDeferLocalizationDataLoad)
+			{
+				IPluginManager::Get().UnmountExplicitlyLoadedPluginLocalizationData(StateProperties.PluginName);
+			}
+
 			// The asset registry listens to FPackageName::OnContentPathDismounted() and 
 			// will automatically cleanup the asset registry state we added for this plugin.
 			// This will also cause any assets we added to the asset manager to be removed.
@@ -1995,6 +2005,10 @@ struct FGameFeaturePluginState_Mounting : public FGameFeaturePluginState
 		if (!UseAsyncLoading() || UE::GameFeatures::CVarForceSyncLoadShaderLibrary.GetValueOnGameThread())
 		{
 			verify(IPluginManager::Get().MountExplicitlyLoadedPlugin(StateProperties.PluginName));
+			if (!UE::GameFeatures::bDeferLocalizationDataLoad)
+			{
+				IPluginManager::Get().MountExplicitlyLoadedPluginLocalizationData(StateProperties.PluginName);
+			}
 			if (bManuallyOpenPluginShaderLibrary)
 			{
 				TSharedPtr<IPlugin> Plugin = IPluginManager::Get().FindPlugin(StateProperties.PluginName);
@@ -2005,6 +2019,10 @@ struct FGameFeaturePluginState_Mounting : public FGameFeaturePluginState
 		}
 
 		verify(IPluginManager::Get().MountExplicitlyLoadedPlugin(StateProperties.PluginName));
+		if (!UE::GameFeatures::bDeferLocalizationDataLoad)
+		{
+			IPluginManager::Get().MountExplicitlyLoadedPluginLocalizationData(StateProperties.PluginName);
+		}
 
 		// Now load the shader lib in the background
 		TSharedPtr<IPlugin> Plugin = IPluginManager::Get().FindPlugin(StateProperties.PluginName);
@@ -2606,6 +2624,10 @@ struct FGameFeaturePluginState_Unloading : public FGameFeaturePluginState
 
 	virtual void BeginState() override
 	{
+		if (UE::GameFeatures::bDeferLocalizationDataLoad)
+		{
+			IPluginManager::Get().UnmountExplicitlyLoadedPluginLocalizationData(StateProperties.PluginName);
+		}
 	}
 
 	virtual void UpdateState(FGameFeaturePluginStateStatus& StateStatus) override
@@ -2654,6 +2676,11 @@ struct FGameFeaturePluginState_Loading : public FGameFeaturePluginState
 		TRACE_CPUPROFILER_EVENT_SCOPE(GFP_Loading_Begin);
 		check(StateProperties.GameFeatureData);
 
+		if (UE::GameFeatures::bDeferLocalizationDataLoad)
+		{
+			IPluginManager::Get().MountExplicitlyLoadedPluginLocalizationData(StateProperties.PluginName);
+		}
+
 		BundleHandle = LoadGameFeatureBundles(StateProperties.GameFeatureData);
 		if (BundleHandle)
 		{
@@ -2883,8 +2910,6 @@ struct FGameFeaturePluginState_Deactivating : public FGameFeaturePluginState
 			{
 				StateStatus.SetTransition(EGameFeaturePluginState::Loaded);
 			}
-
-			IPluginManager::Get().UnmountExplicitlyLoadedPluginLocalizationData(StateProperties.PluginName);
 		}
 		else
 		{
@@ -2962,8 +2987,6 @@ struct FGameFeaturePluginState_Activating : public FGameFeaturePluginState
 
 		FGameFeatureActivatingContext Context;
 
-		const bool bIsLoadingLocalizationData = IPluginManager::Get().MountExplicitlyLoadedPluginLocalizationData(StateProperties.PluginName);
-
 		if (AllowIniLoading())
 		{
 			TRACE_CPUPROFILER_EVENT_SCOPE(GFP_Activating_InitIni);
@@ -2977,7 +3000,8 @@ struct FGameFeaturePluginState_Activating : public FGameFeaturePluginState
 
 		// @TODO: non-blocking wait here?
 		// If this plugin caused localization data to load, wait for that here before marking it as active
-		if (bIsLoadingLocalizationData)
+		if (TSharedPtr<IPlugin> Plugin = IPluginManager::Get().FindPlugin(StateProperties.PluginName);
+			Plugin && Plugin->GetDescriptor().bExplicitlyLoaded && Plugin->GetDescriptor().LocalizationTargets.Num() > 0)
 		{
 			TRACE_CPUPROFILER_EVENT_SCOPE(GFP_Activating_WaitForLoc);
 			FTextLocalizationManager::Get().WaitForAsyncTasks();
