FGameFeaturePluginState_AssetDependencyStreaming now requests to stream asset dependencies from install bundle manager.

GFP Subsys
- GetPluginURLByName now takes a string view
- FindPluginAssetDependencies now returns a list of pairs with the plugin name as the key

Game Feature Project Policy
- Add a function to resolve asset depenencies to a list of install bundles

diff --git a/Engine/Plugins/Runtime/GameFeatures/Source/GameFeatures/Private/GameFeaturePluginStateMachine.cpp b/Engine/Plugins/Runtime/GameFeatures/Source/GameFeatures/Private/GameFeaturePluginStateMachine.cpp
--- a/Engine/Plugins/Runtime/GameFeatures/Source/GameFeatures/Private/GameFeaturePluginStateMachine.cpp
+++ b/Engine/Plugins/Runtime/GameFeatures/Source/GameFeatures/Private/GameFeaturePluginStateMachine.cpp
@@ -2269,23 +2269,123 @@ struct FGameFeaturePluginState_WaitingForDependencies : public FTransitionDepend
 
 struct FGameFeaturePluginState_AssetDependencyStreaming : public FGameFeaturePluginState
 {
-	FGameFeaturePluginState_AssetDependencyStreaming(FGameFeaturePluginStateMachineProperties& InStateProperties) : FGameFeaturePluginState(InStateProperties) {}
+	FGameFeaturePluginState_AssetDependencyStreaming(FGameFeaturePluginStateMachineProperties& InStateProperties) 
+		: FGameFeaturePluginState(InStateProperties) 
+		, Result(MakeValue())
+	{}
+
+	~FGameFeaturePluginState_AssetDependencyStreaming()
+	{
+		Cleanup();
+	}
+
+	void Cleanup()
+	{
+		Result = MakeValue();
+		PendingBundleDownloads.Empty();
+	}
+
+	UE::GameFeatures::FResult Result;
+	TArray<FName> PendingBundleDownloads;
 
-	TArray<FString> AssetDependencies;
+	// TODO: How can we handle progress for this, we don't know we are going to do it until way after downloading
+	// TUniquePtr<FInstallBundleCombinedProgressTracker> ProgressTracker;
+
+	void OnInstallBundleCompleted(FInstallBundleRequestResultInfo BundleResult)
+	{
+		if (!PendingBundleDownloads.Contains(BundleResult.BundleName))
+		{
+			return;
+		}
+
+		PendingBundleDownloads.Remove(BundleResult.BundleName);
+
+		if (!Result.HasError() && BundleResult.Result != EInstallBundleResult::OK)
+		{
+			//Use OptionalErrorCode and/or OptionalErrorText if available
+			const FString ErrorCodeEnding = (BundleResult.OptionalErrorCode.IsEmpty()) ? LexToString(BundleResult.Result) : BundleResult.OptionalErrorCode;
+			const FText ErrorText = BundleResult.OptionalErrorCode.IsEmpty() ? UE::GameFeatures::CommonErrorCodes::GetErrorTextForBundleResult(BundleResult.Result) : BundleResult.OptionalErrorText;
+
+			Result = GetErrorResult(TEXT("BundleManager.AssetDep."), ErrorCodeEnding, ErrorText);
+
+			if (BundleResult.Result != EInstallBundleResult::UserCancelledError)
+			{
+				TSharedPtr<IInstallBundleManager> BundleManager = IInstallBundleManager::GetPlatformInstallBundleManager();
+				BundleManager->CancelUpdateContent(PendingBundleDownloads);
+			}
+		}
+
+		if (PendingBundleDownloads.Num() > 0)
+		{
+			return;
+		}
+
+		UpdateStateMachineImmediate();
+	}
 
 	virtual void BeginState() override
 	{
-		AssetDependencies = UGameFeaturesSubsystem::Get().FindPluginAssetDependencies(StateProperties.PluginInstalledFilename);
+		Cleanup();
+
+		// TODO: Install Bundles need to move away from FNames for identifiers, this is currently just bloating up the name table
+		// when using dynamic GFPs
+		TArray<TPair<FString, TArray<FString>>> AssetDependencies = UGameFeaturesSubsystem::Get().FindPluginAssetDependencies(StateProperties.PluginInstalledFilename);
+		TArray<FName> AssetInstallBundles = UGameFeaturesSubsystem::Get().GetPolicy().GetStreamingAssetInstallBundles(AssetDependencies);
+		if (AssetInstallBundles.IsEmpty())
+		{
+			return;
+		}
+
+		// Respect DoNotDownload flag
+		// TODO: move DoNotDownload flag to base protocol options? We could now have a file: GFP that needs to stream dependencies
+		if (StateProperties.ProtocolOptions.HasSubtype<FInstallBundlePluginProtocolOptions>())
+		{
+			const FInstallBundlePluginProtocolOptions& Options = StateProperties.ProtocolOptions.GetSubtype<FInstallBundlePluginProtocolOptions>();
+			if (Options.bDoNotDownload)
+			{
+				Result = GetErrorResult(TEXT("GFPStateMachine.DownloadNotAllowed"));
+				return;
+			}
+		}
+
+		TSharedPtr<IInstallBundleManager> BundleManager = IInstallBundleManager::GetPlatformInstallBundleManager();
+
+		// TODO: where do flags come from? The dependency being streamed from?
+		EInstallBundleRequestFlags InstallFlags = EInstallBundleRequestFlags::None;
+		TValueOrError<FInstallBundleRequestInfo, EInstallBundleResult> MaybeRequestInfo = BundleManager->RequestUpdateContent(AssetInstallBundles, InstallFlags);
+		if (MaybeRequestInfo.HasError())
+		{
+			const FStringView ShortUrl = StateProperties.PluginIdentifier.GetIdentifyingString();
+			ensureMsgf(false, TEXT("Unable to enqueue asset dependencies for the PluginURL(%.*s) because %s"), ShortUrl.Len(), ShortUrl.GetData(), LexToString(MaybeRequestInfo.GetError()));
+			Result = GetErrorResult(TEXT("BundleManager.AssetDep."), LexToString(MaybeRequestInfo.GetError()));
+			return;
+		}
+
+		FInstallBundleRequestInfo RequestInfo = MaybeRequestInfo.StealValue();
+		PendingBundleDownloads = MoveTemp(RequestInfo.BundlesEnqueued);
+		IInstallBundleManager::InstallBundleCompleteDelegate.AddRaw(this, &FGameFeaturePluginState_AssetDependencyStreaming::OnInstallBundleCompleted);
+
+		// TODO: how to apply pausing? bUserPauseDownload is protocol specific
+		//if (Options.bUserPauseDownload)
 	}
 
 	virtual void UpdateState(FGameFeaturePluginStateStatus& StateStatus) override
 	{
-		StateStatus.SetTransition(EGameFeaturePluginState::Registering);
+		if (!Result.HasValue())
+		{
+			StateStatus.SetTransitionError(EGameFeaturePluginState::ErrorWaitingForDependencies, Result);
+			return;
+		}
+
+		if (PendingBundleDownloads.IsEmpty())
+		{
+			StateStatus.SetTransition(EGameFeaturePluginState::Registering);
+		}
 	}
 
 	virtual void EndState() override
 	{
-		AssetDependencies.Empty();
+		Cleanup();
 	}
 
 	virtual void TryCancelState() override
diff --git a/Engine/Plugins/Runtime/GameFeatures/Source/GameFeatures/Private/GameFeaturesSubsystem.cpp b/Engine/Plugins/Runtime/GameFeatures/Source/GameFeatures/Private/GameFeaturesSubsystem.cpp
--- a/Engine/Plugins/Runtime/GameFeatures/Source/GameFeatures/Private/GameFeaturesSubsystem.cpp
+++ b/Engine/Plugins/Runtime/GameFeatures/Source/GameFeatures/Private/GameFeaturesSubsystem.cpp
@@ -1996,9 +1996,9 @@ void UGameFeaturesSubsystem::LoadBuiltInGameFeaturePlugins(FBuiltInPluginAdditio
 	}
 }
 
-bool UGameFeaturesSubsystem::GetPluginURLByName(const FString& PluginName, FString& OutPluginURL) const
+bool UGameFeaturesSubsystem::GetPluginURLByName(FStringView PluginName, FString& OutPluginURL) const
 {
-	if (const FString* PluginURL = GameFeaturePluginNameToPathMap.Find(PluginName))
+	if (const FString* PluginURL = GameFeaturePluginNameToPathMap.FindByHash(GetTypeHash(PluginName), PluginName))
 	{
 		OutPluginURL = *PluginURL;
 		return true;
@@ -2788,21 +2788,19 @@ void UGameFeaturesSubsystem::FinishTermination(UGameFeaturePluginStateMachine* M
 	TerminalGameFeaturePluginStateMachines.RemoveSwap(Machine);
 }
 
-TArray<FString> UGameFeaturesSubsystem::FindPluginAssetDependencies(const FString& PluginDescriptorFilename)
+TArray<TPair<FString, TArray<FString>>> UGameFeaturesSubsystem::FindPluginAssetDependencies(const FString& PluginDescriptorFilename)
 {
 	FGameFeaturePluginDetails Details;
 	ensure(GetGameFeaturePluginDetailsInternal(PluginDescriptorFilename, Details));
 
-	TArray<FString> OutAssetPaths;
-	for (const FGameFeaturePluginReferenceDetails& RefDetails : Details.PluginDependencies)
+	TArray<TPair<FString, TArray<FString>>> OutDeps;
+	OutDeps.Reserve(Details.PluginDependencies.Num());
+	for (FGameFeaturePluginReferenceDetails& RefDetails : Details.PluginDependencies)
 	{
-		for (const FString& AssetRef : RefDetails.AssetReferences)
-		{
-			OutAssetPaths.Emplace(FString::Printf(TEXT("/%s/%s"), *RefDetails.PluginName, *AssetRef));
-		}
+		OutDeps.Emplace(MoveTemp(RefDetails.PluginName), MoveTemp(RefDetails.AssetReferences));
 	}
 
-	return OutAssetPaths;
+	return OutDeps;
 }
 
 bool UGameFeaturesSubsystem::FindOrCreatePluginDependencyStateMachines(const FString& PluginURL, const FGameFeaturePluginStateMachineProperties& InStateProperties, TArray<UGameFeaturePluginStateMachine*>& OutDependencyMachines)
@@ -3085,6 +3083,7 @@ void UGameFeaturesSubsystem::CallbackObservers(EObserverCallback CallbackType, c
 	case EObserverCallback::PreMounting:
 	{
 		check(PluginName);
+		check(StateChangeContext);
 		FGameFeaturePreMountingContext* PreMountingContext = static_cast<FGameFeaturePreMountingContext*>(StateChangeContext);
 		for (UObject* Observer : LocalObservers)
 		{
@@ -3095,6 +3094,7 @@ void UGameFeaturesSubsystem::CallbackObservers(EObserverCallback CallbackType, c
 	case EObserverCallback::PostMounting:
 	{
 		check(PluginName);
+		check(StateChangeContext);
 		FGameFeaturePostMountingContext* PostMountingContext = static_cast<FGameFeaturePostMountingContext*>(StateChangeContext);
 		for (UObject* Observer : LocalObservers)
 		{
diff --git a/Engine/Plugins/Runtime/GameFeatures/Source/GameFeatures/Public/GameFeaturesProjectPolicies.h b/Engine/Plugins/Runtime/GameFeatures/Source/GameFeatures/Public/GameFeaturesProjectPolicies.h
--- a/Engine/Plugins/Runtime/GameFeatures/Source/GameFeatures/Public/GameFeaturesProjectPolicies.h
+++ b/Engine/Plugins/Runtime/GameFeatures/Source/GameFeatures/Public/GameFeaturesProjectPolicies.h
@@ -41,11 +41,15 @@ public:
 	// (e.g., when doing a fast cook a game might want to disable some or all game feature plugins)
 	virtual bool IsPluginAllowed(const FString& PluginURL) const { return true; }
 
-	// Called to resolve plugin dependencies, OutDependencyURL will contain the GFP URL if this dependency is a GFP.
+	// Called to resolve plugin dependencies, will successfully return an empty string if a dependency is not a GFP.
 	// This may be called with file protocol for built-in plugins in some cases, even if a different protocol is used at runtime.
 	// returns The dependency URL or an error if the dependency could not be resolved
 	virtual TValueOrError<FString, FString> ResolvePluginDependency(const FString& PluginURL, const FString& DependencyName) const;
 
+	// Called to resolve install bundles for stream asset dependencies, input maps a plugin dependency to referenced assets 
+	// (See UGameFeaturesSubsystem::FindPluginAssetDependencies)
+	virtual TArray<FName> GetStreamingAssetInstallBundles(const TArray<TPair<FString, TArray<FString>>>& AssetDependencies) const { return {}; }
+
 	// Called by code that explicitly wants to load a specific plugin
 	// (e.g., when using a fast cook a game might want to allow explicitly loaded game feature plugins)
 	virtual void ExplicitLoadGameFeaturePlugin(const FString& PluginURL, const FGameFeaturePluginLoadComplete& CompleteDelegate, const bool bActivateGameFeatures);
diff --git a/Engine/Plugins/Runtime/GameFeatures/Source/GameFeatures/Public/GameFeaturesSubsystem.h b/Engine/Plugins/Runtime/GameFeatures/Source/GameFeatures/Public/GameFeaturesSubsystem.h
--- a/Engine/Plugins/Runtime/GameFeatures/Source/GameFeatures/Public/GameFeaturesSubsystem.h
+++ b/Engine/Plugins/Runtime/GameFeatures/Source/GameFeatures/Public/GameFeaturesSubsystem.h
@@ -562,7 +562,7 @@ public:
 	 * If the specified plugin is known by the game feature system, returns the URL used to identify it
 	 * @return true if the plugin exists, false if it was not found
 	 */
-	bool GetPluginURLByName(const FString& PluginName, FString& OutPluginURL) const;
+	bool GetPluginURLByName(FStringView PluginName, FString& OutPluginURL) const;
 
 	/** If the specified plugin is a built-in plugin, return the URL used to identify it. Returns true if the plugin exists, false if it was not found */
 	UE_DEPRECATED(5.1, "Use GetPluginURLByName instead")
@@ -729,7 +729,7 @@ private:
 	void FinishTermination(UGameFeaturePluginStateMachine* Machine);
 	friend class UGameFeaturePluginStateMachine;
 
-	TArray<FString> FindPluginAssetDependencies(const FString& PluginDescriptorFilename);
+	TArray<TPair<FString, TArray<FString>>> FindPluginAssetDependencies(const FString& PluginDescriptorFilename);
 	friend struct FGameFeaturePluginState_AssetDependencyStreaming;
 
 	/** Handler for when a state machine requests its dependencies. Returns false if the dependencies could not be read */
