Review feedback changes
- Add FGameFeaturePluginDependency alias
- GetStreamingAssetInstallBundles can return an error

diff --git a/Engine/Plugins/Runtime/GameFeatures/Source/GameFeatures/Private/GameFeaturePluginStateMachine.cpp b/Engine/Plugins/Runtime/GameFeatures/Source/GameFeatures/Private/GameFeaturePluginStateMachine.cpp
--- a/Engine/Plugins/Runtime/GameFeatures/Source/GameFeatures/Private/GameFeaturePluginStateMachine.cpp
+++ b/Engine/Plugins/Runtime/GameFeatures/Source/GameFeatures/Private/GameFeaturePluginStateMachine.cpp
@@ -2306,7 +2306,7 @@ struct FGameFeaturePluginState_AssetDependencyStreaming : public FGameFeaturePlu
 			const FString ErrorCodeEnding = (BundleResult.OptionalErrorCode.IsEmpty()) ? LexToString(BundleResult.Result) : BundleResult.OptionalErrorCode;
 			const FText ErrorText = BundleResult.OptionalErrorCode.IsEmpty() ? UE::GameFeatures::CommonErrorCodes::GetErrorTextForBundleResult(BundleResult.Result) : BundleResult.OptionalErrorText;
 
-			Result = GetErrorResult(TEXT("BundleManager.AssetDep."), ErrorCodeEnding, ErrorText);
+			Result = GetErrorResult(TEXT("BundleManager."), ErrorCodeEnding, ErrorText);
 
 			if (BundleResult.Result != EInstallBundleResult::UserCancelledError)
 			{
@@ -2329,8 +2329,15 @@ struct FGameFeaturePluginState_AssetDependencyStreaming : public FGameFeaturePlu
 
 		// TODO: Install Bundles need to move away from FNames for identifiers, this is currently just bloating up the name table
 		// when using dynamic GFPs
-		TArray<TPair<FString, TArray<FString>>> AssetDependencies = UGameFeaturesSubsystem::Get().FindPluginAssetDependencies(StateProperties.PluginInstalledFilename);
-		TArray<FName> AssetInstallBundles = UGameFeaturesSubsystem::Get().GetPolicy().GetStreamingAssetInstallBundles(AssetDependencies);
+		TArray<FGameFeaturePluginDependency> AssetDependencies = UGameFeaturesSubsystem::Get().FindPluginAssetDependencies(StateProperties.PluginInstalledFilename);
+		TValueOrError<TArray<FName>, FString> MaybeAssetInstallBundles = UGameFeaturesSubsystem::Get().GetPolicy().GetStreamingAssetInstallBundles(AssetDependencies);
+		if (MaybeAssetInstallBundles.HasError())
+		{
+			Result = GetErrorResult(MaybeAssetInstallBundles.GetError());
+			return;
+		}
+
+		const TArray<FName>& AssetInstallBundles = MaybeAssetInstallBundles.GetValue();
 		if (AssetInstallBundles.IsEmpty())
 		{
 			return;
@@ -2357,7 +2364,7 @@ struct FGameFeaturePluginState_AssetDependencyStreaming : public FGameFeaturePlu
 		{
 			const FStringView ShortUrl = StateProperties.PluginIdentifier.GetIdentifyingString();
 			ensureMsgf(false, TEXT("Unable to enqueue asset dependencies for the PluginURL(%.*s) because %s"), ShortUrl.Len(), ShortUrl.GetData(), LexToString(MaybeRequestInfo.GetError()));
-			Result = GetErrorResult(TEXT("BundleManager.AssetDep."), LexToString(MaybeRequestInfo.GetError()));
+			Result = GetErrorResult(TEXT("BundleManager."), LexToString(MaybeRequestInfo.GetError()));
 			return;
 		}
 
diff --git a/Engine/Plugins/Runtime/GameFeatures/Source/GameFeatures/Private/GameFeaturesSubsystem.cpp b/Engine/Plugins/Runtime/GameFeatures/Source/GameFeatures/Private/GameFeaturesSubsystem.cpp
--- a/Engine/Plugins/Runtime/GameFeatures/Source/GameFeatures/Private/GameFeaturesSubsystem.cpp
+++ b/Engine/Plugins/Runtime/GameFeatures/Source/GameFeatures/Private/GameFeaturesSubsystem.cpp
@@ -2788,12 +2788,12 @@ void UGameFeaturesSubsystem::FinishTermination(UGameFeaturePluginStateMachine* M
 	TerminalGameFeaturePluginStateMachines.RemoveSwap(Machine);
 }
 
-TArray<TPair<FString, TArray<FString>>> UGameFeaturesSubsystem::FindPluginAssetDependencies(const FString& PluginDescriptorFilename)
+TArray<FGameFeaturePluginDependency> UGameFeaturesSubsystem::FindPluginAssetDependencies(const FString& PluginDescriptorFilename)
 {
 	FGameFeaturePluginDetails Details;
 	ensure(GetGameFeaturePluginDetailsInternal(PluginDescriptorFilename, Details));
 
-	TArray<TPair<FString, TArray<FString>>> OutDeps;
+	TArray<FGameFeaturePluginDependency> OutDeps;
 	OutDeps.Reserve(Details.PluginDependencies.Num());
 	for (FGameFeaturePluginReferenceDetails& RefDetails : Details.PluginDependencies)
 	{
diff --git a/Engine/Plugins/Runtime/GameFeatures/Source/GameFeatures/Public/GameFeaturesProjectPolicies.h b/Engine/Plugins/Runtime/GameFeatures/Source/GameFeatures/Public/GameFeaturesProjectPolicies.h
--- a/Engine/Plugins/Runtime/GameFeatures/Source/GameFeatures/Public/GameFeaturesProjectPolicies.h
+++ b/Engine/Plugins/Runtime/GameFeatures/Source/GameFeatures/Public/GameFeaturesProjectPolicies.h
@@ -48,7 +48,7 @@ public:
 
 	// Called to resolve install bundles for stream asset dependencies, input maps a plugin dependency to referenced assets 
 	// (See UGameFeaturesSubsystem::FindPluginAssetDependencies)
-	virtual TArray<FName> GetStreamingAssetInstallBundles(const TArray<TPair<FString, TArray<FString>>>& AssetDependencies) const { return {}; }
+	virtual TValueOrError<TArray<FName>, FString> GetStreamingAssetInstallBundles(const TArray<FGameFeaturePluginDependency>& AssetDependencies) const { return MakeValue(); }
 
 	// Called by code that explicitly wants to load a specific plugin
 	// (e.g., when using a fast cook a game might want to allow explicitly loaded game feature plugins)
diff --git a/Engine/Plugins/Runtime/GameFeatures/Source/GameFeatures/Public/GameFeaturesSubsystem.h b/Engine/Plugins/Runtime/GameFeatures/Source/GameFeatures/Public/GameFeaturesSubsystem.h
--- a/Engine/Plugins/Runtime/GameFeatures/Source/GameFeatures/Public/GameFeaturesSubsystem.h
+++ b/Engine/Plugins/Runtime/GameFeatures/Source/GameFeatures/Public/GameFeaturesSubsystem.h
@@ -386,6 +386,9 @@ struct FGameFeatureInfo
 	EGameFeaturePluginState CurrentState;
 };
 
+// Key - plugin reference, Value - any specific assets imported
+using FGameFeaturePluginDependency = TPair<FString, TArray<FString>>;
+
 /** The manager subsystem for game features */
 UCLASS()
 class GAMEFEATURES_API UGameFeaturesSubsystem : public UEngineSubsystem
@@ -729,7 +732,7 @@ private:
 	void FinishTermination(UGameFeaturePluginStateMachine* Machine);
 	friend class UGameFeaturePluginStateMachine;
 
-	TArray<TPair<FString, TArray<FString>>> FindPluginAssetDependencies(const FString& PluginDescriptorFilename);
+	TArray<FGameFeaturePluginDependency> FindPluginAssetDependencies(const FString& PluginDescriptorFilename);
 	friend struct FGameFeaturePluginState_AssetDependencyStreaming;
 
 	/** Handler for when a state machine requests its dependencies. Returns false if the dependencies could not be read */
