UGameFeatureActionConvertContentBundleWorldPartitionBuilder fixes (ContentBundle to EDL Converter)
- Simplified builder by reusing new helper method FExternalDataLayerHelper::MoveActorsToExternalDataLayers to move actors from Content Bundle to EDL
- Updated FExternalDataLayerHelper::MoveActorsToExternalDataLayers to support moving all dependant objects inside OFPA package (mainly HLOD actors)
- Removed the builder option '-skipdelete' as it is not viable to keep both sets of actor since the AssetRegistry would parse 2 different pacakges for the same actor path are choose the last scanned
- Fixed deleted by closing old OFPA file handles before deleting the old content bundle packages
- Store the source ContentBundle Guid in the External Data Layer GameFeatureAction to facilitate any potential future partial revert of the convert operation

diff --git a/Engine/Plugins/Runtime/GameFeatures/Source/GameFeatures/Public/GameFeatureAction_AddWorldPartitionContent.h b/Engine/Plugins/Runtime/GameFeatures/Source/GameFeatures/Public/GameFeatureAction_AddWorldPartitionContent.h
--- a/Engine/Plugins/Runtime/GameFeatures/Source/GameFeatures/Public/GameFeatureAction_AddWorldPartitionContent.h
+++ b/Engine/Plugins/Runtime/GameFeatures/Source/GameFeatures/Public/GameFeatureAction_AddWorldPartitionContent.h
@@ -47,5 +47,11 @@ private:
 	UPROPERTY(EditAnywhere, Category = DataLayer)
 	TObjectPtr<const UExternalDataLayerAsset> ExternalDataLayerAsset;
 
+#if WITH_EDITORONLY_DATA
+	/** Only used when converting from UGameFeatureAction_AddWPContent */
+	UPROPERTY()
+	FGuid ConvertedContentBundleGuid;
+#endif
+
 	friend class UGameFeatureActionConvertContentBundleWorldPartitionBuilder;
 };
\ No newline at end of file
diff --git a/Engine/Plugins/Runtime/GameFeatures/Source/GameFeaturesEditor/Private/GameFeatureActionConvertContentBundleWorldPartitionBuilder.cpp b/Engine/Plugins/Runtime/GameFeatures/Source/GameFeaturesEditor/Private/GameFeatureActionConvertContentBundleWorldPartitionBuilder.cpp
--- a/Engine/Plugins/Runtime/GameFeatures/Source/GameFeaturesEditor/Private/GameFeatureActionConvertContentBundleWorldPartitionBuilder.cpp
+++ b/Engine/Plugins/Runtime/GameFeatures/Source/GameFeaturesEditor/Private/GameFeatureActionConvertContentBundleWorldPartitionBuilder.cpp
@@ -30,7 +30,6 @@ DEFINE_LOG_CATEGORY_STATIC(LogConvertContentBundleBuilder, All, All);
 UGameFeatureActionConvertContentBundleWorldPartitionBuilder::UGameFeatureActionConvertContentBundleWorldPartitionBuilder(const FObjectInitializer& ObjectInitializer)
 	: Super(ObjectInitializer)
 	, bReportOnly(false)
-	, bSkipDelete(false)
 {}
 
 bool UGameFeatureActionConvertContentBundleWorldPartitionBuilder::PreRun(UWorld* World, FPackageSourceControlHelper& PackageHelper)
@@ -42,7 +41,6 @@ bool UGameFeatureActionConvertContentBundleWorldPartitionBuilder::PreRun(UWorld*
 	FString const* DestFolderPtr = CommandLineParams.Find(TEXT("DestinationFolder"));
 	DestinationFolder = DestFolderPtr ? *DestFolderPtr : FString();
 	bReportOnly = Switches.Contains(TEXT("ReportOnly"));
-	bSkipDelete = Switches.Contains(TEXT("SkipDelete"));
 
 	if (FString const* ContentBundlesToConvertString = CommandLineParams.Find(TEXT("ContentBundles")))
 	{
@@ -105,7 +103,7 @@ bool UGameFeatureActionConvertContentBundleWorldPartitionBuilder::RunInternal(UW
 	for (TSharedPtr<FContentBundleEditor>& ContentBundle : ContentBundlesToProcess)
 	{
 		TSet<UPackage*> PackagesToSave;
-		TSet<FString> PackagesToDelete;
+		TSet<UPackage*> PackagesToDelete;
 
 		TSharedPtr<FContentBundleClient> ContentBundleClient = ContentBundle->GetClient().Pin();
 		if (!ContentBundleClient.IsValid())
@@ -182,6 +180,10 @@ bool UGameFeatureActionConvertContentBundleWorldPartitionBuilder::RunInternal(UW
 			UE_LOG(LogConvertContentBundleBuilder, Log, TEXT("Added new Action of type 'GameFeatureAction_AddWorldPartitionContent' to GameFeatureData %s using External Data Layer Asset %s while converting Content Bundle %s."), *GameFeatureData->GetName(), *ExternalDataLayerAsset->GetPathName(), *ContentBundle->GetDisplayName());
 		}
 		
+		// Backup the source ContentBundle Guid (to facilitate potential revert operation)
+		AddWorldPartitionContent->ConvertedContentBundleGuid = ContentBundleDescriptor->GetGuid();
+		PackagesToSave.Add(GameFeatureData->GetPackage());
+		
 		// Manually call OnExternalDataLayerAssetChanged to register the newly created GameFeatureAction_AddWorldPartitionContent
 		AddWorldPartitionContent->OnExternalDataLayerAssetChanged(nullptr, ExternalDataLayerAsset);
 
@@ -238,67 +240,23 @@ bool UGameFeatureActionConvertContentBundleWorldPartitionBuilder::RunInternal(UW
 			ActorReferences.Add(ActorRef);
 			AActor* Actor = ActorRef.GetActor();
 			UPackage* OldActorPackage = Actor->GetExternalPackage();
-			const FString OldActorPackageName = OldActorPackage->GetName();
-			if (!bSkipDelete)
-			{
-				PackagesToDelete.Add(OldActorPackageName);
-			}
 
-			// Clear Actor CB Guid
-			FSetActorContentBundleGuid(Actor, FGuid());
-
-			// Verify that the actor can be assigned to the External Data Layer
 			FText FailureReason;
-			if (!ExternalDataLayerInstance->CanAddActor(Actor, &FailureReason))
+			if (!FExternalDataLayerHelper::MoveActorsToExternalDataLayer({ Actor }, ExternalDataLayerInstance, &FailureReason))
 			{
 				UE_LOG(LogConvertContentBundleBuilder, Error, TEXT("Can't create package for actor %s. %s"), *ExternalDataLayerAsset->GetName(), *Actor->GetActorNameOrLabel(), *FailureReason.ToString());
 				bContentBundleActorConversionSuccess = false;
 				break;
 			}
 
-			// Remove actor from it's old package
-			Actor->SetPackageExternal(false);
-
-			// Get all other dependant objects in the old actor package
-			TArray<UObject*> DependantObjects;
-			ForEachObjectWithPackage(OldActorPackage, [&DependantObjects](UObject* Object)
-			{
-				if (!Cast<UMetaData>(Object))
-				{
-					DependantObjects.Add(Object);
-				}
-				return true;
-			}, false);
-
-			// Create a new external package and assign it to the actor
-			ULevel* DestinationLevel = Actor->GetLevel();
-			UPackage* NewActorPackage = ULevel::CreateActorPackage(DestinationLevel->GetPackage(), DestinationLevel->GetActorPackagingScheme(), Actor->GetName(), ExternalDataLayerAsset);
-			Actor->SetPackageExternal(true, true, NewActorPackage);
-
-			// Validation
-			check(NewActorPackage == Actor->GetExternalPackage());
-			check(NewActorPackage->GetName() == ExternalDataLayerManager->GetActorPackageName(ExternalDataLayerAsset, DestinationLevel, Actor->GetPathName()));
-			check(NewActorPackage->GetName() != OldActorPackageName);
-
-			// Move dependant objects into the new actor package
-			for (UObject* DependantObject : DependantObjects)
-			{
-				DependantObject->Rename(nullptr, NewActorPackage, REN_NonTransactional | REN_DontCreateRedirectors | REN_ForceNoResetLoaders | REN_DoNotDirty);
-			}
-			
-			// Set External Data Layer Asset
-			if (!UDataLayerEditorSubsystem::Get()->AddActorToDataLayer(Actor, ExternalDataLayerInstance))
-			{
-				UE_LOG(LogConvertContentBundleBuilder, Error, TEXT("Failed to add actor %s(%s) to external data layer %s. Actor won't be converted."), *It->GetActorLabelOrName().ToString(), *It->GetActorPackage().ToString(), *ExternalDataLayerAsset->GetName());
-				bContentBundleActorConversionSuccess = false;
-				break;
-			}
-
 			check(Actor->GetExternalDataLayerAsset() == ExternalDataLayerAsset);
 			check(!Actor->GetContentBundleGuid().IsValid());
 			UPackage* NewPackage = Actor->GetPackage();
 			PackagesToSave.Add(NewPackage);
-			check(OldActorPackageName != NewPackage->GetName());
+
+			check(OldActorPackage->GetName() != NewPackage->GetName());
+			PackagesToDelete.Add(OldActorPackage);
+			
 			UE_LOG(LogConvertContentBundleBuilder, Log, TEXT("Converted Actor %s(%s) to %s."), *Actor->GetName(), *NewPackage->GetName(), *ExternalDataLayerAsset->GetName());
 
 			ConvertedActorGuids.Add(It->GetGuid());
@@ -329,9 +287,10 @@ bool UGameFeatureActionConvertContentBundleWorldPartitionBuilder::RunInternal(UW
 			UE_LOG(LogConvertContentBundleBuilder, Log, TEXT("Package to save: %s"), *PackageToSave->GetPathName());
 		}
 
-		for (const FString& PackageToDelete : PackagesToDelete)
+		for (UPackage* PackageToDelete : PackagesToDelete)
 		{
-			UE_LOG(LogConvertContentBundleBuilder, Log, TEXT("Package to delete: %s"), *PackageToDelete);
+			UE_LOG(LogConvertContentBundleBuilder, Log, TEXT("Package to delete: %s"), *PackageToDelete->GetPathName());
+			ResetLoaders(PackageToDelete);
 		}
 
 		if (bContentBundleActorConversionSuccess && !bReportOnly)
@@ -363,9 +322,9 @@ bool UGameFeatureActionConvertContentBundleWorldPartitionBuilder::RunInternal(UW
 			if (PackagesToDelete.Num())
 			{
 				FinalReport.Add(FString::Printf(TEXT("[+] Deleted %d packages: "), PackagesToDelete.Num()));
-				for (const FString& PackageToDelete : PackagesToDelete)
+				for (UPackage* PackageToDelete : PackagesToDelete)
 				{
-					FinalReport.Add(FString::Printf(TEXT(" |- %s"), *PackageToDelete));
+					FinalReport.Add(FString::Printf(TEXT(" |- %s"), *PackageToDelete->GetPathName()));
 				}
 			}
 		}
diff --git a/Engine/Plugins/Runtime/GameFeatures/Source/GameFeaturesEditor/Private/GameFeatureActionConvertContentBundleWorldPartitionBuilder.h b/Engine/Plugins/Runtime/GameFeatures/Source/GameFeaturesEditor/Private/GameFeatureActionConvertContentBundleWorldPartitionBuilder.h
--- a/Engine/Plugins/Runtime/GameFeatures/Source/GameFeaturesEditor/Private/GameFeatureActionConvertContentBundleWorldPartitionBuilder.h
+++ b/Engine/Plugins/Runtime/GameFeatures/Source/GameFeaturesEditor/Private/GameFeatureActionConvertContentBundleWorldPartitionBuilder.h
@@ -32,5 +32,4 @@ private:
 	TArray<FString> ContentBundlesToConvert;
 	FString DestinationFolder;
 	bool bReportOnly;
-	bool bSkipDelete;
 };
\ No newline at end of file
diff --git a/Engine/Source/Runtime/Engine/Private/WorldPartition/DataLayer/ExternalDataLayerHelper.cpp b/Engine/Source/Runtime/Engine/Private/WorldPartition/DataLayer/ExternalDataLayerHelper.cpp
--- a/Engine/Source/Runtime/Engine/Private/WorldPartition/DataLayer/ExternalDataLayerHelper.cpp
+++ b/Engine/Source/Runtime/Engine/Private/WorldPartition/DataLayer/ExternalDataLayerHelper.cpp
@@ -9,9 +9,11 @@
 #include "AssetRegistry/ARFilter.h"
 #include "ExternalPackageHelper.h"
 #include "UObject/Object.h"
+#include "UObject/MetaData.h"
 #include "UObject/AssetRegistryTagsContext.h"
 #include "Subsystems/EditorActorSubsystem.h"
 #include "WorldPartition/DataLayer/ExternalDataLayerManager.h"
+#include "DeletedObjectPlaceholder.h"
 #include "ReferencedAssetsUtils.h"
 #include "Editor.h"
 #endif
@@ -239,13 +241,25 @@ bool FExternalDataLayerHelper::CanMoveActorsToExternalDataLayer(const TArray<AAc
 
 bool FExternalDataLayerHelper::MoveActorsToExternalDataLayer(const TArray<AActor*>& InActors, const UExternalDataLayerInstance* InExternalDataLayerInstance, FText* OutFailureReason)
 {
-	auto MoveActorToExternalDataLayer = [](AActor * InActor, const UExternalDataLayerInstance* InExternalDataLayerInstance)
+	auto MoveActorToExternalDataLayer = [](AActor* InActor, const UExternalDataLayerInstance* InExternalDataLayerInstance)
 	{
+		const UPackage* OldActorPackage = InActor->GetExternalPackage();
 		const UExternalDataLayerAsset* NewExternalDataLayerAsset = InExternalDataLayerInstance ? InExternalDataLayerInstance->GetExternalDataLayerAsset() : nullptr;
 		const bool bShouldDirty = true;
-		bool bLevelPackageWasDirty = InActor->GetLevel()->GetPackage()->IsDirty();
+		const bool bLevelPackageWasDirty = InActor->GetLevel()->GetPackage()->IsDirty();
 		InActor->SetPackageExternal(false, bShouldDirty);
 
+		// Get all other dependant objects in the old actor package
+		TArray<UObject*> DependantObjects;
+		ForEachObjectWithPackage(OldActorPackage, [&DependantObjects](UObject* Object)
+		{
+			if (!Cast<UMetaData>(Object) && !Cast<UDeletedObjectPlaceholder>(Object))
+			{
+				DependantObjects.Add(Object);
+			}
+			return true;
+		}, false);
+
 		// Clear Content Bundle Guid
 		FSetActorContentBundleGuid(InActor, FGuid());
 
@@ -270,6 +284,14 @@ bool FExternalDataLayerHelper::MoveActorsToExternalDataLayer(const TArray<AActor
 		}
 
 		InActor->SetPackageExternal(true, bShouldDirty);
+
+		// Move dependant objects into the new actor package
+		UPackage* NewActorPackage = InActor->GetExternalPackage();
+		for (UObject* DependantObject : DependantObjects)
+		{
+			DependantObject->Rename(nullptr, NewActorPackage, REN_NonTransactional | REN_DontCreateRedirectors | REN_ForceNoResetLoaders | REN_DoNotDirty);
+		}
+
 		if (!bLevelPackageWasDirty)
 		{
 			InActor->GetLevel()->GetPackage()->SetDirtyFlag(false);
diff --git a/Engine/Source/Runtime/Engine/Private/WorldPartition/DataLayer/ExternalDataLayerManager.cpp b/Engine/Source/Runtime/Engine/Private/WorldPartition/DataLayer/ExternalDataLayerManager.cpp
--- a/Engine/Source/Runtime/Engine/Private/WorldPartition/DataLayer/ExternalDataLayerManager.cpp
+++ b/Engine/Source/Runtime/Engine/Private/WorldPartition/DataLayer/ExternalDataLayerManager.cpp
@@ -827,8 +827,7 @@ void UExternalDataLayerManager::ForEachExternalStreamingObjects(TFunctionRef<boo
 
 FString UExternalDataLayerManager::GetActorPackageName(const UExternalDataLayerAsset* InExternalDataLayerAsset, const ULevel* InDestinationLevel, const FString& InActorPath) const
 {
-	const FString ContainerRootPath = GetExternalDataLayerLevelRootPath(InExternalDataLayerAsset);
-	const FString ActorPackageName = ULevel::GetActorPackageName(ULevel::GetExternalActorsPath(ContainerRootPath), InDestinationLevel->GetActorPackagingScheme(), InActorPath);
+	const FString ActorPackageName = ULevel::GetActorPackageName(InDestinationLevel->GetPackage(), InDestinationLevel->GetActorPackagingScheme(), InActorPath, InExternalDataLayerAsset);
 	return ActorPackageName;
 }
 
diff --git a/Engine/Source/Runtime/Engine/Public/WorldPartition/DataLayer/ExternalDataLayerHelper.h b/Engine/Source/Runtime/Engine/Public/WorldPartition/DataLayer/ExternalDataLayerHelper.h
--- a/Engine/Source/Runtime/Engine/Public/WorldPartition/DataLayer/ExternalDataLayerHelper.h
+++ b/Engine/Source/Runtime/Engine/Public/WorldPartition/DataLayer/ExternalDataLayerHelper.h
@@ -70,6 +70,7 @@ private:
 	friend class UDataLayerEditorSubsystem;
 	friend class UExternalDataLayerEngineSubsystem;
 	friend class UContentBundleEditingSubmodule;
+	friend class UGameFeatureActionConvertContentBundleWorldPartitionBuilder;
 #endif
 	static constexpr FStringView GetExternalDataLayerFolder() { return ExternalDataLayerFolder; }
 	static constexpr FStringView ExternalDataLayerFolder = TEXTVIEW("/EDL/");
