PR #11073: Factor FGameFeaturePluginTemplateDescription into it's own file

diff --git a/Engine/Plugins/Runtime/GameFeatures/Source/GameFeaturesEditor/GameFeaturesEditor.Build.cs b/Engine/Plugins/Runtime/GameFeatures/Source/GameFeaturesEditor/GameFeaturesEditor.Build.cs
--- a/Engine/Plugins/Runtime/GameFeatures/Source/GameFeaturesEditor/GameFeaturesEditor.Build.cs
+++ b/Engine/Plugins/Runtime/GameFeatures/Source/GameFeaturesEditor/GameFeaturesEditor.Build.cs
@@ -11,7 +11,9 @@ namespace UnrealBuildTool.Rules
                 {
                     "Core",
                     "CoreUObject",
-                }
+					"GameFeatures",
+					"UnrealEd",
+				}
 			);
 
 			PrivateDependencyModuleNames.AddRange(
@@ -24,9 +26,7 @@ namespace UnrealBuildTool.Rules
 					"DeveloperSettings",
 					"Engine",
 					"ModularGameplay",
-					"GameFeatures",
 					"EditorSubsystem",
-					"UnrealEd",
 					"Projects",
 					"EditorFramework",
 					"Slate",
diff --git a/Engine/Plugins/Runtime/GameFeatures/Source/GameFeaturesEditor/Private/GameFeaturePluginTemplate.cpp b/Engine/Plugins/Runtime/GameFeatures/Source/GameFeaturesEditor/Private/GameFeaturePluginTemplate.cpp
new file mode 100644
--- /dev/null
+++ b/Engine/Plugins/Runtime/GameFeatures/Source/GameFeaturesEditor/Private/GameFeaturePluginTemplate.cpp
@@ -0,0 +1,131 @@
+// Copyright Epic Games, Inc. All Rights Reserved.
+
+#include "GameFeaturePluginTemplate.h"
+
+#include "AssetRegistry/ARFilter.h"
+#include "AssetRegistry/IAssetRegistry.h"
+#include "AssetToolsModule.h"
+#include "DOM/JsonValue.h"
+#include "Editor.h"
+#include "GameFeaturesSubsystem.h"
+#include "GameFeaturesSubsystemSettings.h"
+#include "HAL/FileManager.h"
+#include "Interfaces/IPluginManager.h"
+#include "Misc/Paths.h"
+#include "Subsystems/AssetEditorSubsystem.h"
+
+#define LOCTEXT_NAMESPACE "GameFeatures"
+
+FGameFeaturePluginTemplateDescription::FGameFeaturePluginTemplateDescription(FText InName, FText InDescription, FString InOnDiskPath, FString InDefaultSubfolder, TSubclassOf<UGameFeatureData> GameFeatureDataClassOverride, FString GameFeatureDataNameOverride, EPluginEnabledByDefault InEnabledByDefault)
+	: FPluginTemplateDescription(InName, InDescription, InOnDiskPath, /*bCanContainContent=*/ true, EHostType::Runtime)
+{
+	SortPriority = 10;
+	bCanBePlacedInEngine = false;
+	DefaultSubfolder = InDefaultSubfolder;
+	GameFeatureDataName = !GameFeatureDataNameOverride.IsEmpty() ? GameFeatureDataNameOverride : FString();
+	GameFeatureDataClass = GameFeatureDataClassOverride != nullptr ? GameFeatureDataClassOverride : TSubclassOf<UGameFeatureData>(UGameFeatureData::StaticClass());
+	PluginEnabledByDefault = InEnabledByDefault;
+}
+
+bool FGameFeaturePluginTemplateDescription::ValidatePathForPlugin(const FString& ProposedAbsolutePluginPath, FText& OutErrorMessage)
+{
+	if (!IsRootedInGameFeaturesRoot(ProposedAbsolutePluginPath))
+	{
+		OutErrorMessage = LOCTEXT("InvalidPathForGameFeaturePlugin", "Game features must be inside the Plugins/GameFeatures folder");
+		return false;
+	}
+
+	OutErrorMessage = FText::GetEmpty();
+	return true;
+}
+
+void FGameFeaturePluginTemplateDescription::UpdatePathWhenTemplateSelected(FString& InOutPath)
+{
+	if (!IsRootedInGameFeaturesRoot(InOutPath))
+	{
+		InOutPath = GetGameFeatureRoot();
+	}
+}
+
+void FGameFeaturePluginTemplateDescription::UpdatePathWhenTemplateUnselected(FString& InOutPath)
+{
+	InOutPath = IFileManager::Get().ConvertToAbsolutePathForExternalAppForWrite(*FPaths::ProjectPluginsDir());
+	FPaths::MakePlatformFilename(InOutPath);
+}
+
+void FGameFeaturePluginTemplateDescription::CustomizeDescriptorBeforeCreation(FPluginDescriptor& Descriptor)
+{
+	Descriptor.bExplicitlyLoaded = true;
+	Descriptor.AdditionalFieldsToWrite.FindOrAdd(TEXT("BuiltInInitialFeatureState")) = MakeShared<FJsonValueString>(TEXT("Active"));
+	Descriptor.Category = TEXT("Game Features");
+
+	// Game features should not be enabled by default if the game wants to strictly manage default settings in the target settings
+	Descriptor.EnabledByDefault = PluginEnabledByDefault;
+
+	if (Descriptor.Modules.Num() > 0)
+	{
+		Descriptor.Modules[0].Name = FName(*(Descriptor.Modules[0].Name.ToString() + TEXT("Runtime")));
+	}
+}
+
+void FGameFeaturePluginTemplateDescription::OnPluginCreated(TSharedPtr<IPlugin> NewPlugin)
+{
+	// If the template includes an existing game feature data, do not create a new one.
+	TArray<FAssetData> ObjectList;
+	FARFilter AssetFilter;
+	AssetFilter.ClassPaths.Add(UGameFeatureData::StaticClass()->GetClassPathName());
+	AssetFilter.PackagePaths.Add(FName(NewPlugin->GetMountedAssetPath()));
+	AssetFilter.bRecursiveClasses = true;
+	AssetFilter.bRecursivePaths = true;
+
+	IAssetRegistry::GetChecked().GetAssets(AssetFilter, ObjectList);
+
+	UObject* GameFeatureDataAsset = nullptr;
+
+	if (ObjectList.Num() <= 0)
+	{
+		// Create the game feature data asset
+		FAssetToolsModule& AssetToolsModule = FModuleManager::Get().LoadModuleChecked<FAssetToolsModule>("AssetTools");
+		FString const& AssetName = !GameFeatureDataName.IsEmpty() ? GameFeatureDataName : NewPlugin->GetName();
+		GameFeatureDataAsset = AssetToolsModule.Get().CreateAsset(AssetName, NewPlugin->GetMountedAssetPath(), GameFeatureDataClass, /*Factory=*/ nullptr);
+	}
+	else
+	{
+		GameFeatureDataAsset = ObjectList[0].GetAsset();
+	}
+
+
+	// Activate the new game feature plugin
+	auto AdditionalFilter = [](const FString&, const FGameFeaturePluginDetails&, FBuiltInGameFeaturePluginBehaviorOptions&) -> bool { return true; };
+	UGameFeaturesSubsystem::Get().LoadBuiltInGameFeaturePlugin(NewPlugin.ToSharedRef(), AdditionalFilter,
+		FGameFeaturePluginLoadComplete::CreateLambda([GameFeatureDataAsset](const UE::GameFeatures::FResult&)
+			{
+				// Edit the new game feature data
+				if (GameFeatureDataAsset != nullptr)
+				{
+					GEditor->GetEditorSubsystem<UAssetEditorSubsystem>()->OpenEditorForAsset(GameFeatureDataAsset);
+				}
+			}));
+}
+
+FString FGameFeaturePluginTemplateDescription::GetGameFeatureRoot() const
+{
+	FString Result = IFileManager::Get().ConvertToAbsolutePathForExternalAppForWrite(*(FPaths::ProjectPluginsDir() / TEXT("GameFeatures/")));
+
+	// Append the optional subfolder if specified.
+	if (!DefaultSubfolder.IsEmpty())
+	{
+		Result /= DefaultSubfolder + TEXT("/");
+	}
+
+	FPaths::MakePlatformFilename(Result);
+	return Result;
+}
+
+bool FGameFeaturePluginTemplateDescription::IsRootedInGameFeaturesRoot(const FString& InStr) const
+{
+	const FString ConvertedPath = FPaths::ConvertRelativePathToFull(FPaths::CreateStandardFilename(InStr / TEXT("test.uplugin")));
+	return GetDefault<UGameFeaturesSubsystemSettings>()->IsValidGameFeaturePlugin(ConvertedPath);
+}
+
+#undef LOCTEXT_NAMESPACE
\ No newline at end of file
diff --git a/Engine/Plugins/Runtime/GameFeatures/Source/GameFeaturesEditor/Private/GameFeaturesEditorModule.cpp b/Engine/Plugins/Runtime/GameFeatures/Source/GameFeaturesEditor/Private/GameFeaturesEditorModule.cpp
--- a/Engine/Plugins/Runtime/GameFeatures/Source/GameFeaturesEditor/Private/GameFeaturesEditorModule.cpp
+++ b/Engine/Plugins/Runtime/GameFeatures/Source/GameFeaturesEditor/Private/GameFeaturesEditorModule.cpp
@@ -8,154 +8,24 @@
 
 #include "GameFeaturesSubsystem.h"
 #include "Interfaces/IPluginManager.h"
-#include "AssetRegistry/AssetRegistryModule.h"
-#include "AssetToolsModule.h"
-#include "IAssetTools.h"
-#include "Subsystems/AssetEditorSubsystem.h"
+#include "Features/IPluginsEditorFeature.h"
 #include "GameFeatureDataDetailsCustomization.h"
 #include "GameFeaturesEditorSettings.h"
 #include "GameFeaturesSubsystemSettings.h"
 #include "GameFeaturePluginMetadataCustomization.h"
+#include "GameFeaturePluginTemplate.h"
 #include "Logging/MessageLog.h"
-#include "Misc/Paths.h"
+#include "Misc/App.h"
+#include "Modules/ModuleManager.h"
 #include "SSettingsEditorCheckoutNotice.h"
 #include "Engine/AssetManagerSettings.h"
 #include "Widgets/Notifications/SNotificationList.h"
 #include "Framework/Notifications/NotificationManager.h"
 #include "Engine/AssetManager.h"
-#include "HAL/FileManager.h"
-#include "Editor.h"
-#include "Dom/JsonValue.h"
 #include "PropertyEditorModule.h"
 
 #define LOCTEXT_NAMESPACE "GameFeatures"
 
-//////////////////////////////////////////////////////////////////////
-
-struct FGameFeaturePluginTemplateDescription : public FPluginTemplateDescription
-{
-	FGameFeaturePluginTemplateDescription(FText InName, FText InDescription, FString InOnDiskPath, FString InDefaultSubfolder
-		, TSubclassOf<UGameFeatureData> GameFeatureDataClassOverride, FString GameFeatureDataNameOverride, EPluginEnabledByDefault InEnabledByDefault)
-		: FPluginTemplateDescription(InName, InDescription, InOnDiskPath, /*bCanContainContent=*/ true, EHostType::Runtime)
-	{
-		SortPriority = 10;
-		bCanBePlacedInEngine = false;
-		DefaultSubfolder = InDefaultSubfolder;
-		GameFeatureDataName = !GameFeatureDataNameOverride.IsEmpty() ? GameFeatureDataNameOverride : FString();
-		GameFeatureDataClass = GameFeatureDataClassOverride != nullptr ? GameFeatureDataClassOverride : TSubclassOf<UGameFeatureData>(UGameFeatureData::StaticClass());
-		PluginEnabledByDefault = InEnabledByDefault;
-	}
-
-	virtual bool ValidatePathForPlugin(const FString& ProposedAbsolutePluginPath, FText& OutErrorMessage) override
-	{
-		if (!IsRootedInGameFeaturesRoot(ProposedAbsolutePluginPath))
-		{
-			OutErrorMessage = LOCTEXT("InvalidPathForGameFeaturePlugin", "Game features must be inside the Plugins/GameFeatures folder");
-			return false;
-		}
-
-		OutErrorMessage = FText::GetEmpty();
-		return true;
-	}
-
-	virtual void UpdatePathWhenTemplateSelected(FString& InOutPath) override
-	{
-		if (!IsRootedInGameFeaturesRoot(InOutPath))
-		{
-			InOutPath = GetGameFeatureRoot();
-		}
-	}
-
-	virtual void UpdatePathWhenTemplateUnselected(FString& InOutPath) override
-	{
-		InOutPath = IFileManager::Get().ConvertToAbsolutePathForExternalAppForWrite(*FPaths::ProjectPluginsDir());
-		FPaths::MakePlatformFilename(InOutPath);
-	}
-
-	virtual void CustomizeDescriptorBeforeCreation(FPluginDescriptor& Descriptor) override
-	{
-		Descriptor.bExplicitlyLoaded = true;
-		Descriptor.AdditionalFieldsToWrite.FindOrAdd(TEXT("BuiltInInitialFeatureState")) = MakeShared<FJsonValueString>(TEXT("Active"));
-		Descriptor.Category = TEXT("Game Features");
-
-		// Game features should not be enabled by default if the game wants to strictly manage default settings in the target settings
-		Descriptor.EnabledByDefault = PluginEnabledByDefault;
-
-		if (Descriptor.Modules.Num() > 0)
-		{
-			Descriptor.Modules[0].Name = FName(*(Descriptor.Modules[0].Name.ToString() + TEXT("Runtime")));
-		}
-	}
-
-	virtual void OnPluginCreated(TSharedPtr<IPlugin> NewPlugin) override
-	{
-		// If the template includes an existing game feature data, do not create a new one.
-		FAssetRegistryModule& AssetRegistryModule = FModuleManager::LoadModuleChecked<FAssetRegistryModule>(TEXT("AssetRegistry"));
-		TArray<FAssetData> ObjectList;
-		FARFilter AssetFilter;
-		AssetFilter.ClassPaths.Add(UGameFeatureData::StaticClass()->GetClassPathName());
-		AssetFilter.PackagePaths.Add(FName(NewPlugin->GetMountedAssetPath()));
-		AssetFilter.bRecursiveClasses = true;
-		AssetFilter.bRecursivePaths = true;
-
-		AssetRegistryModule.Get().GetAssets(AssetFilter, ObjectList);
-
-		UObject* GameFeatureDataAsset = nullptr;
-
-		if (ObjectList.Num() <= 0)
-		{
-			// Create the game feature data asset
-			FAssetToolsModule& AssetToolsModule = FModuleManager::Get().LoadModuleChecked<FAssetToolsModule>("AssetTools");
-			FString const& AssetName = !GameFeatureDataName.IsEmpty() ? GameFeatureDataName : NewPlugin->GetName();
-			GameFeatureDataAsset = AssetToolsModule.Get().CreateAsset(AssetName, NewPlugin->GetMountedAssetPath(), GameFeatureDataClass, /*Factory=*/ nullptr);
-		}
-		else
-		{
-			GameFeatureDataAsset = ObjectList[0].GetAsset();
-		}
-		
-
-		// Activate the new game feature plugin
-		auto AdditionalFilter = [](const FString&, const FGameFeaturePluginDetails&, FBuiltInGameFeaturePluginBehaviorOptions&) -> bool { return true; };
-		UGameFeaturesSubsystem::Get().LoadBuiltInGameFeaturePlugin(NewPlugin.ToSharedRef(), AdditionalFilter,
-			FGameFeaturePluginLoadComplete::CreateLambda([GameFeatureDataAsset](const UE::GameFeatures::FResult&)
-			{
-				// Edit the new game feature data
-				if (GameFeatureDataAsset != nullptr)
-				{
-					GEditor->GetEditorSubsystem<UAssetEditorSubsystem>()->OpenEditorForAsset(GameFeatureDataAsset);
-				}
-			}));
-	}
-
-	FString GetGameFeatureRoot() const
-	{
-		FString Result = IFileManager::Get().ConvertToAbsolutePathForExternalAppForWrite(*(FPaths::ProjectPluginsDir() / TEXT("GameFeatures/")));
-
-		// Append the optional subfolder if specified.
-		if (!DefaultSubfolder.IsEmpty())
-		{
-			Result /= DefaultSubfolder + TEXT("/");
-		}
-
-		FPaths::MakePlatformFilename(Result);
-		return Result;
-	}
-
-	bool IsRootedInGameFeaturesRoot(const FString& InStr)
-	{
-		const FString ConvertedPath = FPaths::ConvertRelativePathToFull(FPaths::CreateStandardFilename(InStr / TEXT("test.uplugin")));
-		return GetDefault<UGameFeaturesSubsystemSettings>()->IsValidGameFeaturePlugin(ConvertedPath);
-	}
-
-	FString DefaultSubfolder;
-	TSubclassOf<UGameFeatureData> GameFeatureDataClass;
-	FString GameFeatureDataName;
-	EPluginEnabledByDefault PluginEnabledByDefault = EPluginEnabledByDefault::Disabled;
-};
-
-//////////////////////////////////////////////////////////////////////
-
 class FGameFeaturesEditorModule : public FDefaultModuleImpl
 {
 	virtual void StartupModule() override
diff --git a/Engine/Plugins/Runtime/GameFeatures/Source/GameFeaturesEditor/Public/GameFeaturePluginTemplate.h b/Engine/Plugins/Runtime/GameFeatures/Source/GameFeaturesEditor/Public/GameFeaturePluginTemplate.h
new file mode 100644
--- /dev/null
+++ b/Engine/Plugins/Runtime/GameFeatures/Source/GameFeaturesEditor/Public/GameFeaturePluginTemplate.h
@@ -0,0 +1,29 @@
+// Copyright Epic Games, Inc. All Rights Reserved.
+
+#include "Features/IPluginsEditorFeature.h"
+#include "GameFeatureData.h"
+#include "PluginDescriptor.h"
+
+/**
+ * Used to create custom templates for GameFeaturePlugins.
+ */
+struct GAMEFEATURESEDITOR_API FGameFeaturePluginTemplateDescription : public FPluginTemplateDescription
+{
+	FGameFeaturePluginTemplateDescription(FText InName, FText InDescription, FString InOnDiskPath, FString InDefaultSubfolder
+		, TSubclassOf<UGameFeatureData> GameFeatureDataClassOverride, FString GameFeatureDataNameOverride, EPluginEnabledByDefault InEnabledByDefault);
+
+	virtual bool ValidatePathForPlugin(const FString& ProposedAbsolutePluginPath, FText& OutErrorMessage) override;
+	virtual void UpdatePathWhenTemplateSelected(FString& InOutPath) override;
+	virtual void UpdatePathWhenTemplateUnselected(FString& InOutPath) override;
+
+	virtual void CustomizeDescriptorBeforeCreation(FPluginDescriptor& Descriptor) override;
+	virtual void OnPluginCreated(TSharedPtr<IPlugin> NewPlugin) override;
+
+	FString GetGameFeatureRoot() const;
+	bool IsRootedInGameFeaturesRoot(const FString& InStr) const;
+
+	FString DefaultSubfolder;
+	TSubclassOf<UGameFeatureData> GameFeatureDataClass;
+	FString GameFeatureDataName;
+	EPluginEnabledByDefault PluginEnabledByDefault = EPluginEnabledByDefault::Disabled;
+};
