Missing pragma once from 32865412

diff --git a/Engine/Plugins/Runtime/GameFeatures/Source/GameFeaturesEditor/Public/GameFeaturePluginTemplate.h b/Engine/Plugins/Runtime/GameFeatures/Source/GameFeaturesEditor/Public/GameFeaturePluginTemplate.h
--- a/Engine/Plugins/Runtime/GameFeatures/Source/GameFeaturesEditor/Public/GameFeaturePluginTemplate.h
+++ b/Engine/Plugins/Runtime/GameFeatures/Source/GameFeaturesEditor/Public/GameFeaturePluginTemplate.h
@@ -1,5 +1,7 @@
 // Copyright Epic Games, Inc. All Rights Reserved.
 
+#pragma once
+
 #include "Features/IPluginsEditorFeature.h"
 #include "GameFeatureData.h"
 #include "PluginDescriptor.h"
