Fixing linux compile error due to wrong include casing

diff --git a/Engine/Plugins/Runtime/GameFeatures/Source/GameFeaturesEditor/Private/GameFeaturePluginTemplate.cpp b/Engine/Plugins/Runtime/GameFeatures/Source/GameFeaturesEditor/Private/GameFeaturePluginTemplate.cpp
--- a/Engine/Plugins/Runtime/GameFeatures/Source/GameFeaturesEditor/Private/GameFeaturePluginTemplate.cpp
+++ b/Engine/Plugins/Runtime/GameFeatures/Source/GameFeaturesEditor/Private/GameFeaturePluginTemplate.cpp
@@ -5,7 +5,7 @@
 #include "AssetRegistry/ARFilter.h"
 #include "AssetRegistry/IAssetRegistry.h"
 #include "AssetToolsModule.h"
-#include "DOM/JsonValue.h"
+#include "Dom/JsonValue.h"
 #include "Editor.h"
 #include "GameFeaturesSubsystem.h"
 #include "GameFeaturesSubsystemSettings.h"
