Don't defer GFP localization data loading in the editor

The editor only needs to mount plugins to use them, so we should load the localization data on 'mounting' when GIsEditor is true

#rnx

diff --git a/Engine/Plugins/Runtime/GameFeatures/Source/GameFeatures/Private/GameFeaturePluginStateMachine.cpp b/Engine/Plugins/Runtime/GameFeatures/Source/GameFeatures/Private/GameFeaturePluginStateMachine.cpp
--- a/Engine/Plugins/Runtime/GameFeatures/Source/GameFeatures/Private/GameFeaturePluginStateMachine.cpp
+++ b/Engine/Plugins/Runtime/GameFeatures/Source/GameFeatures/Private/GameFeaturePluginStateMachine.cpp
@@ -95,6 +95,12 @@ namespace UE::GameFeatures
 		false,
 		TEXT("Enable experimental asset streaming"));
 
+	bool ShouldDeferLocalizationDataLoad()
+	{
+		// Note: We don't defer localization data loading in the editor, as the editor only needs to mount plugins to use them
+		return !GIsEditor && bDeferLocalizationDataLoad;
+	}
+
 	bool ShouldSkipVerify(const FString& PluginName)
 	{
 		static const FAsciiSet Wildcards("*?");
@@ -1622,7 +1628,7 @@ struct FGameFeaturePluginState_Unmounting : public FGameFeaturePluginState
 		if (TSharedPtr<IPlugin> Plugin = IPluginManager::Get().FindPlugin(StateProperties.PluginName);
 			Plugin && Plugin->GetDescriptor().bExplicitlyLoaded)
 		{
-			if (!UE::GameFeatures::bDeferLocalizationDataLoad)
+			if (!UE::GameFeatures::ShouldDeferLocalizationDataLoad())
 			{
 				IPluginManager::Get().UnmountExplicitlyLoadedPluginLocalizationData(StateProperties.PluginName);
 			}
@@ -2005,7 +2011,7 @@ struct FGameFeaturePluginState_Mounting : public FGameFeaturePluginState
 		if (!UseAsyncLoading() || UE::GameFeatures::CVarForceSyncLoadShaderLibrary.GetValueOnGameThread())
 		{
 			verify(IPluginManager::Get().MountExplicitlyLoadedPlugin(StateProperties.PluginName));
-			if (!UE::GameFeatures::bDeferLocalizationDataLoad)
+			if (!UE::GameFeatures::ShouldDeferLocalizationDataLoad())
 			{
 				IPluginManager::Get().MountExplicitlyLoadedPluginLocalizationData(StateProperties.PluginName);
 			}
@@ -2019,7 +2025,7 @@ struct FGameFeaturePluginState_Mounting : public FGameFeaturePluginState
 		}
 
 		verify(IPluginManager::Get().MountExplicitlyLoadedPlugin(StateProperties.PluginName));
-		if (!UE::GameFeatures::bDeferLocalizationDataLoad)
+		if (!UE::GameFeatures::ShouldDeferLocalizationDataLoad())
 		{
 			IPluginManager::Get().MountExplicitlyLoadedPluginLocalizationData(StateProperties.PluginName);
 		}
@@ -2731,7 +2737,7 @@ struct FGameFeaturePluginState_Unloading : public FGameFeaturePluginState
 
 	virtual void BeginState() override
 	{
-		if (UE::GameFeatures::bDeferLocalizationDataLoad)
+		if (UE::GameFeatures::ShouldDeferLocalizationDataLoad())
 		{
 			IPluginManager::Get().UnmountExplicitlyLoadedPluginLocalizationData(StateProperties.PluginName);
 		}
@@ -2783,7 +2789,7 @@ struct FGameFeaturePluginState_Loading : public FGameFeaturePluginState
 		TRACE_CPUPROFILER_EVENT_SCOPE(GFP_Loading_Begin);
 		check(StateProperties.GameFeatureData);
 
-		if (UE::GameFeatures::bDeferLocalizationDataLoad)
+		if (UE::GameFeatures::ShouldDeferLocalizationDataLoad())
 		{
 			IPluginManager::Get().MountExplicitlyLoadedPluginLocalizationData(StateProperties.PluginName);
 		}
