Add EInstallBundleResult::MetadataError

diff --git a/Engine/Plugins/Runtime/GameFeatures/Source/GameFeatures/Private/GameFeaturePluginOperationResult.cpp b/Engine/Plugins/Runtime/GameFeatures/Source/GameFeatures/Private/GameFeaturePluginOperationResult.cpp
--- a/Engine/Plugins/Runtime/GameFeatures/Source/GameFeatures/Private/GameFeaturePluginOperationResult.cpp
+++ b/Engine/Plugins/Runtime/GameFeatures/Source/GameFeatures/Private/GameFeaturePluginOperationResult.cpp
@@ -66,6 +66,7 @@ namespace UE::GameFeatures
 
 				//These are generally unrecoverable and mean something is seriously wrong with the data for this build
 				case EInstallBundleResult::InitializationError:
+				case EInstallBundleResult::MetadataError:
 				{
 					return Generic_FatalError;
 				}
@@ -109,12 +110,13 @@ namespace UE::GameFeatures
 		{
 			switch (ErrorResult)
 			{
-				case (EInstallBundleReleaseResult::UserCancelledError):
+				case EInstallBundleReleaseResult::UserCancelledError:
 				{
 					return ReleaseResult_Cancelled;
 				}
 			
-				case (EInstallBundleReleaseResult::ManifestArchiveError):
+				case EInstallBundleReleaseResult::ManifestArchiveError:
+				case EInstallBundleReleaseResult::MetadataError:
 				{
 					return ReleaseResult_Generic;
 				}
diff --git a/Engine/Source/Runtime/InstallBundleManager/Private/InstallBundleTypes.cpp b/Engine/Source/Runtime/InstallBundleManager/Private/InstallBundleTypes.cpp
--- a/Engine/Source/Runtime/InstallBundleManager/Private/InstallBundleTypes.cpp
+++ b/Engine/Source/Runtime/InstallBundleManager/Private/InstallBundleTypes.cpp
@@ -163,6 +163,7 @@ const TCHAR* LexToString(EInstallBundleResult Result)
 		TEXT("UserCancelledError"),
 		TEXT("InitializationError"),
 		TEXT("InitializationPending"),
+		TEXT("MetadataError"),
 	};
 
 	return InstallBundleUtil::TLexToString(Result, Strings);
@@ -175,6 +176,7 @@ const TCHAR* LexToString(EInstallBundleReleaseResult Result)
 		TEXT("OK"),
 		TEXT("ManifestArchiveError"),
 		TEXT("UserCancelledError"),
+		TEXT("MetadataError"),
 	};
 
 	return InstallBundleUtil::TLexToString(Result, Strings);
diff --git a/Engine/Source/Runtime/InstallBundleManager/Public/InstallBundleTypes.h b/Engine/Source/Runtime/InstallBundleManager/Public/InstallBundleTypes.h
--- a/Engine/Source/Runtime/InstallBundleManager/Public/InstallBundleTypes.h
+++ b/Engine/Source/Runtime/InstallBundleManager/Public/InstallBundleTypes.h
@@ -174,6 +174,7 @@ enum class EInstallBundleResult : uint32
 	UserCancelledError,
 	InitializationError,
 	InitializationPending,
+	MetadataError,
 	Count,
 };
 INSTALLBUNDLEMANAGER_API const TCHAR* LexToString(EInstallBundleResult Result);
@@ -198,6 +199,7 @@ enum class EInstallBundleReleaseResult : uint32
 	OK,
 	ManifestArchiveError,
 	UserCancelledError,
+	MetadataError,
 	Count,
 };
 INSTALLBUNDLEMANAGER_API const TCHAR* LexToString(EInstallBundleReleaseResult Result);
