Fixed GFPs failing to build their package localization cache when loading their asset registry data async

We need to refresh the package localization cache for the GFP after the asset registry data has loaded, as we need the asset registry data to correctly build the package localization cache

diff --git a/Engine/Plugins/Runtime/GameFeatures/Source/GameFeatures/Private/GameFeaturePluginStateMachine.cpp b/Engine/Plugins/Runtime/GameFeatures/Source/GameFeatures/Private/GameFeaturePluginStateMachine.cpp
--- a/Engine/Plugins/Runtime/GameFeatures/Source/GameFeatures/Private/GameFeaturePluginStateMachine.cpp
+++ b/Engine/Plugins/Runtime/GameFeatures/Source/GameFeatures/Private/GameFeaturePluginStateMachine.cpp
@@ -11,6 +11,7 @@
 #include "InstallBundleUtils.h"
 #include "BundlePrereqCombinedStatusHelper.h"
 #include "Interfaces/IPluginManager.h"
+#include "Internationalization/PackageLocalizationManager.h"
 #include "Logging/StructuredLog.h"
 #include "Materials/MaterialInterface.h"
 #include "Misc/AsciiSet.h"
@@ -1798,6 +1799,7 @@ struct FGameFeaturePluginState_Mounting : public FGameFeaturePluginState
 	UE::GameFeatures::FResult Result;
 	ESubState StartedSubStates = ESubState::None;
 	ESubState CompletedSubStates = ESubState::None;
+	bool bLoadedAssetRegistryState = false;
 	bool bCheckedRealtimeMode = false;
 	bool bForceMonolithicShaderLibrary = true;	// use monolithic unless a DLC plugin is chunked
 
@@ -1880,6 +1882,7 @@ struct FGameFeaturePluginState_Mounting : public FGameFeaturePluginState
 		Result = MakeValue();
 		StartedSubStates = ESubState::None;
 		CompletedSubStates = ESubState::None;
+		bLoadedAssetRegistryState = false;
 		bCheckedRealtimeMode = false;
 		bForceMonolithicShaderLibrary = false;
 
@@ -2115,6 +2118,7 @@ struct FGameFeaturePluginState_Mounting : public FGameFeaturePluginState
 				Result = GetErrorResult(TEXT("Failed_To_Load_Plugin_AssetRegistry"));
 			}
 
+			bLoadedAssetRegistryState = true;
 			CompletedSubStates |= ESubState::LoadAssetRegistry;
 			return;
 		}
@@ -2148,6 +2152,7 @@ struct FGameFeaturePluginState_Mounting : public FGameFeaturePluginState
 					AssetRegistry.AppendState(*PluginAssetRegistryState);
 				}
 
+				bLoadedAssetRegistryState = true;
 				CompletedSubStates |= ESubState::LoadAssetRegistry;
 				UpdateStateMachineImmediate();
 			});
@@ -2195,6 +2200,18 @@ struct FGameFeaturePluginState_Mounting : public FGameFeaturePluginState
 		// Post-mount
 		if (bComplete)
 		{
+			if (bLoadedAssetRegistryState)
+			{
+				// We need to refresh the package localization cache for a GFP if it loaded cooked asset registry state, 
+				// as we need the asset registry data to correctly build the package localization cache for the GFP
+				if (TSharedPtr<IPlugin> Plugin = IPluginManager::Get().FindPlugin(StateProperties.PluginName);
+					Plugin && Plugin->CanContainContent())
+				{
+					FPackageLocalizationManager::Get().InvalidateRootSourcePath(Plugin->GetMountedAssetPath());
+					FPackageLocalizationManager::Get().ConditionalUpdateCache();
+				}
+			}
+
 			FGameFeaturePostMountingContext Context(StateProperties.PluginName, [this](FStringView InPauserTag) { OnPostMountPauserCompleted(InPauserTag); });
 			NumExpectedPostMountPausers = INDEX_NONE;
 			UGameFeaturesSubsystem::Get().OnGameFeaturePostMounting(StateProperties.PluginName, StateProperties.PluginIdentifier, Context);
diff --git a/Engine/Source/Runtime/CoreUObject/Private/Internationalization/PackageLocalizationCache.cpp b/Engine/Source/Runtime/CoreUObject/Private/Internationalization/PackageLocalizationCache.cpp
--- a/Engine/Source/Runtime/CoreUObject/Private/Internationalization/PackageLocalizationCache.cpp
+++ b/Engine/Source/Runtime/CoreUObject/Private/Internationalization/PackageLocalizationCache.cpp
@@ -85,9 +85,13 @@ void FPackageLocalizationCultureCache::RemoveRootSourcePath(const FString& InRoo
 	}
 
 	// Remove all packages under this root
+	FNameBuilder SourcePackageName;
 	for (auto It = SourcePackagesToLocalizedPackages.CreateIterator(); It; ++It)
 	{
-		if (It->Key.ToString().StartsWith(InRootPath))
+		SourcePackageName.Reset();
+		It->Key.AppendString(SourcePackageName);
+
+		if (SourcePackageName.ToView().StartsWith(InRootPath))
 		{
 			It.RemoveCurrent();
 			continue;
@@ -235,6 +239,19 @@ FPackageLocalizationCache::~FPackageLocalizationCache()
 	FPackageName::OnContentPathDismounted().RemoveAll(this);
 }
 
+void FPackageLocalizationCache::InvalidateRootSourcePath(const FString& InRootPath)
+{
+	FScopeLock Lock(&LocalizedCachesCS);
+
+	for (auto& CultureCachePair : AllCultureCaches)
+	{
+		CultureCachePair.Value->RemoveRootSourcePath(InRootPath);
+		CultureCachePair.Value->AddRootSourcePath(InRootPath);
+	}
+
+	bPackageNameToAssetGroupDirty = true;
+}
+
 void FPackageLocalizationCache::ConditionalUpdateCache()
 {
 	FScopeLock Lock(&LocalizedCachesCS);
diff --git a/Engine/Source/Runtime/CoreUObject/Private/Internationalization/PackageLocalizationManager.cpp b/Engine/Source/Runtime/CoreUObject/Private/Internationalization/PackageLocalizationManager.cpp
--- a/Engine/Source/Runtime/CoreUObject/Private/Internationalization/PackageLocalizationManager.cpp
+++ b/Engine/Source/Runtime/CoreUObject/Private/Internationalization/PackageLocalizationManager.cpp
@@ -163,6 +163,14 @@ FName FPackageLocalizationManager::FindLocalizedPackageNameNoCache(const FName I
 	return NAME_None;
 }
 
+void FPackageLocalizationManager::InvalidateRootSourcePath(const FString& InRootPath)
+{
+	if (ActiveCache.IsValid())
+	{
+		ActiveCache->InvalidateRootSourcePath(InRootPath);
+	}
+}
+
 void FPackageLocalizationManager::ConditionalUpdateCache()
 {
 	if (ActiveCache.IsValid())
diff --git a/Engine/Source/Runtime/CoreUObject/Public/Internationalization/IPackageLocalizationCache.h b/Engine/Source/Runtime/CoreUObject/Public/Internationalization/IPackageLocalizationCache.h
--- a/Engine/Source/Runtime/CoreUObject/Public/Internationalization/IPackageLocalizationCache.h
+++ b/Engine/Source/Runtime/CoreUObject/Public/Internationalization/IPackageLocalizationCache.h
@@ -10,6 +10,12 @@ class IPackageLocalizationCache
 public:
 	virtual ~IPackageLocalizationCache() {}
 
+	/**
+	 * Invalidate any cached state for the given root source path, and add it to the queue of things to process when ConditionalUpdateCache is called.
+	 * eg) when new asset registry state is loaded for a plugin that may invalidate its cached data
+	 */
+	virtual void InvalidateRootSourcePath(const FString& InRootPath) = 0;
+
 	/**
 	 * Update this cache, but only if it is dirty.
 	 */
diff --git a/Engine/Source/Runtime/CoreUObject/Public/Internationalization/PackageLocalizationCache.h b/Engine/Source/Runtime/CoreUObject/Public/Internationalization/PackageLocalizationCache.h
--- a/Engine/Source/Runtime/CoreUObject/Public/Internationalization/PackageLocalizationCache.h
+++ b/Engine/Source/Runtime/CoreUObject/Public/Internationalization/PackageLocalizationCache.h
@@ -114,6 +114,7 @@ public:
 	COREUOBJECT_API virtual ~FPackageLocalizationCache();
 
 	//~ IPackageLocalizationCache interface
+	COREUOBJECT_API virtual void InvalidateRootSourcePath(const FString& InRootPath) override;
 	COREUOBJECT_API virtual void ConditionalUpdateCache() override;
 	COREUOBJECT_API virtual FName FindLocalizedPackageName(const FName InSourcePackageName) override;
 	COREUOBJECT_API virtual FName FindLocalizedPackageNameForCulture(const FName InSourcePackageName, const FString& InCultureName) override;
diff --git a/Engine/Source/Runtime/CoreUObject/Public/Internationalization/PackageLocalizationManager.h b/Engine/Source/Runtime/CoreUObject/Public/Internationalization/PackageLocalizationManager.h
--- a/Engine/Source/Runtime/CoreUObject/Public/Internationalization/PackageLocalizationManager.h
+++ b/Engine/Source/Runtime/CoreUObject/Public/Internationalization/PackageLocalizationManager.h
@@ -60,6 +60,12 @@ public:
 	 */
 	COREUOBJECT_API FName FindLocalizedPackageNameForCulture(const FName InSourcePackageName, const FString& InCultureName);
 
+	/**
+	 * Invalidate any cached state for the given root source path, and add it to the queue of things to process when ConditionalUpdateCache is called.
+	 * eg) when new asset registry state is loaded for a plugin that may invalidate its cached data
+	 */
+	COREUOBJECT_API void InvalidateRootSourcePath(const FString& InRootPath);
+
 	/**
 	 * Update this cache, but only if it is dirty.
 	 */
