Log all GFP state transitions once their state machine returns to an idle (or error) state

This logs the transition regardless of how it starts (eg, LoadAndActivateGameFeaturePlugin), rather than only transitions that start via LoadBuiltInGameFeaturePlugin

#rnx

diff --git a/Engine/Plugins/Runtime/GameFeatures/Source/GameFeatures/Private/GameFeaturePluginStateMachine.cpp b/Engine/Plugins/Runtime/GameFeatures/Source/GameFeatures/Private/GameFeaturePluginStateMachine.cpp
--- a/Engine/Plugins/Runtime/GameFeatures/Source/GameFeatures/Private/GameFeaturePluginStateMachine.cpp
+++ b/Engine/Plugins/Runtime/GameFeatures/Source/GameFeatures/Private/GameFeaturePluginStateMachine.cpp
@@ -3514,7 +3514,8 @@ UE_TRACE_EVENT_END()
 
 void UGameFeaturePluginStateMachine::UpdateStateMachine()
 {
-	EGameFeaturePluginState CurrentState = GetCurrentState();
+	const EGameFeaturePluginState InitialState = GetCurrentState();
+	EGameFeaturePluginState CurrentState = InitialState;
 	if (bInUpdateStateMachine)
 	{
 		UE_LOG(LogGameFeatures, Verbose, TEXT("Game feature state machine skipping update for %s in ::UpdateStateMachine. Current State: %s"), *GetGameFeatureName(), *UE::GameFeatures::ToString(CurrentState));
@@ -3648,6 +3649,28 @@ void UGameFeaturePluginStateMachine::UpdateStateMachine()
 			}
 		}
 
+		// Log our final state if we've finished transitioning
+		if (!bKeepProcessing && InitialState != CurrentState)
+		{
+			if (!StateStatus.TransitionResult.HasValue())
+			{
+				UE_LOG(LogGameFeatures, Error, TEXT("Game feature '%s' transition failed. Ending state: %s [%s, %s]. Result: %s"),
+					*GetGameFeatureName(),
+					*UE::GameFeatures::ToString(CurrentState),
+					*UE::GameFeatures::ToString(StateProperties.Destination.MinState),
+					*UE::GameFeatures::ToString(StateProperties.Destination.MaxState),
+					*UE::GameFeatures::ToString(StateStatus.TransitionResult));
+			}
+			else if (StateProperties.Destination.Contains(CurrentState))
+			{
+				UE_LOG(LogGameFeatures, Display, TEXT("Game feature '%s' transitioned successfully. Ending state: %s [%s, %s]"),
+					*GetGameFeatureName(),
+					*UE::GameFeatures::ToString(CurrentState),
+					*UE::GameFeatures::ToString(StateProperties.Destination.MinState),
+					*UE::GameFeatures::ToString(StateProperties.Destination.MaxState));
+			}
+		}
+
 		if (NumTransitions++ > MaxTransitions)
 		{
 			UE_LOG(LogGameFeatures, Fatal, TEXT("Infinite loop in game feature state machine transitions. Current state %s. GameFeature: %s"), *UE::GameFeatures::ToString(CurrentState), *GetGameFeatureName());
diff --git a/Engine/Plugins/Runtime/GameFeatures/Source/GameFeatures/Private/GameFeaturesSubsystem.cpp b/Engine/Plugins/Runtime/GameFeatures/Source/GameFeatures/Private/GameFeaturesSubsystem.cpp
--- a/Engine/Plugins/Runtime/GameFeatures/Source/GameFeatures/Private/GameFeaturesSubsystem.cpp
+++ b/Engine/Plugins/Runtime/GameFeatures/Source/GameFeatures/Private/GameFeaturesSubsystem.cpp
@@ -2648,7 +2648,7 @@ UGameFeaturePluginStateMachine* UGameFeaturesSubsystem::FindOrCreateGameFeatureP
 		return *ExistingStateMachine;
 	}
 
-	UE_LOG(LogGameFeatures, Display, TEXT("Creating GameFeaturePlugin StateMachine using Identifier:%.*s from PluginURL:%s"), PluginIdentifier.GetIdentifyingString().Len(), PluginIdentifier.GetIdentifyingString().GetData(), *PluginURL);
+	UE_LOG(LogGameFeatures, Verbose, TEXT("Creating GameFeaturePlugin StateMachine using Identifier:%.*s from PluginURL:%s"), PluginIdentifier.GetIdentifyingString().Len(), PluginIdentifier.GetIdentifyingString().GetData(), *PluginURL);
 
 	UGameFeaturePluginStateMachine* NewStateMachine = NewObject<UGameFeaturePluginStateMachine>(this);
 	GameFeaturePluginStateMachines.Add(FString(PluginIdentifier.GetIdentifyingString()), NewStateMachine);
@@ -2662,38 +2662,15 @@ void UGameFeaturesSubsystem::LoadBuiltInGameFeaturePluginComplete(const UE::Game
 	check(Machine);
 	if (Result.HasValue())
 	{
-		//@note It's possible for the machine to still be tranitioning at this point as long as it's withing the requested destination range
-		UE_LOG(LogGameFeatures, Display, TEXT("Game feature '%s' loaded successfully. Ending state: %s, [%s, %s]"), 
-			*Machine->GetGameFeatureName(), 
-			*UE::GameFeatures::ToString(Machine->GetCurrentState()),
-			*UE::GameFeatures::ToString(Machine->GetDestination().MinState),
-			*UE::GameFeatures::ToString(Machine->GetDestination().MaxState));
-
 		checkf(RequestedDestination.Contains(Machine->GetCurrentState()), TEXT("Game feature '%s': Ending state %s is not in expected range [%s, %s]"), 
 			*Machine->GetGameFeatureName(), 
 			*UE::GameFeatures::ToString(Machine->GetCurrentState()), 
 			*UE::GameFeatures::ToString(RequestedDestination.MinState), 
 			*UE::GameFeatures::ToString(RequestedDestination.MaxState));
-#if WITH_EDITOR
-		if (Machine->GetCurrentState() > EGameFeaturePluginState::Mounting)
-		{
-			if (FString* Explanation = UnmountedPluginNameToExplanation.Find(Machine->GetPluginName()))
-			{
-				UnmountedPluginNameToExplanation.Remove(Machine->GetPluginName());
-			}
-		}
-#endif
 	}
 	else
 	{
-		const FString ErrorMessage = UE::GameFeatures::ToString(Result);
-		UE_LOG(LogGameFeatures, Error, TEXT("Game feature '%s' load failed. Ending state: %s, [%s, %s]. Result: %s"),
-			*Machine->GetGameFeatureName(),
-			*UE::GameFeatures::ToString(Machine->GetCurrentState()),
-			*UE::GameFeatures::ToString(Machine->GetDestination().MinState),
-			*UE::GameFeatures::ToString(Machine->GetDestination().MaxState),
-			*ErrorMessage);
-		SetExplanationForNotMountingPlugin(Machine->GetPluginURL(), ErrorMessage);
+		SetExplanationForNotMountingPlugin(Machine->GetPluginURL(), UE::GameFeatures::ToString(Result));
 	}
 }
 
@@ -2761,10 +2738,7 @@ void UGameFeaturesSubsystem::ChangeGameFeatureTargetStateComplete(UGameFeaturePl
 #if WITH_EDITOR
 	if (!Result.HasError() && Machine->GetCurrentState() > EGameFeaturePluginState::Mounting)
 	{
-		if (FString* Explanation = UnmountedPluginNameToExplanation.Find(Machine->GetPluginName()))
-		{
-			UnmountedPluginNameToExplanation.Remove(Machine->GetPluginName());
-		}
+		UnmountedPluginNameToExplanation.Remove(Machine->GetPluginName());
 	}
 #endif
 	CompleteDelegate.ExecuteIfBound(Result);
@@ -2784,7 +2758,7 @@ void UGameFeaturesSubsystem::BeginTermination(UGameFeaturePluginStateMachine* Ma
 
 void UGameFeaturesSubsystem::FinishTermination(UGameFeaturePluginStateMachine* Machine)
 {
-	UE_LOG(LogGameFeatures, Display, TEXT("FinishTermination of GameFeaturePlugin. Identifier:%.*s URL:%s"), Machine->GetPluginIdentifier().GetIdentifyingString().Len(), Machine->GetPluginIdentifier().GetIdentifyingString().GetData(), *(Machine->GetPluginURL()));
+	UE_LOG(LogGameFeatures, Verbose, TEXT("FinishTermination of GameFeaturePlugin. Identifier:%.*s URL:%s"), Machine->GetPluginIdentifier().GetIdentifyingString().Len(), Machine->GetPluginIdentifier().GetIdentifyingString().GetData(), *(Machine->GetPluginURL()));
 	TerminalGameFeaturePluginStateMachines.RemoveSwap(Machine);
 }
 
