remove unecessary use of lastIndex

diff --git a/Engine/Plugins/Runtime/GameFeatures/Source/GameFeatures/Private/GameFeatureAction_AddChunkOverride.cpp b/Engine/Plugins/Runtime/GameFeatures/Source/GameFeatures/Private/GameFeatureAction_AddChunkOverride.cpp
--- a/Engine/Plugins/Runtime/GameFeatures/Source/GameFeatures/Private/GameFeatureAction_AddChunkOverride.cpp
+++ b/Engine/Plugins/Runtime/GameFeatures/Source/GameFeatures/Private/GameFeatureAction_AddChunkOverride.cpp
@@ -172,7 +172,6 @@ void UGameFeatureAction_AddChunkOverride::AddChunkIdOverride()
 		TArray<FString>& PluginsInChunk = GameFeatureAction_AddChunkOverride::ChunkIdToPluginMap.FindOrAdd(ChunkId);
 		PluginsInChunk.Add(PluginName);
 		GameFeatureAction_AddChunkOverride::PluginToChunkId.Add(PluginName, ChunkId);
-		LastChunkIdUsed = ChunkId;
 		UE_LOG(LogAddChunkOverride, Log, TEXT("Plugin(%s) will cook assets into chunk(%d)"), *PluginName, ChunkId);
 
 		UAssetManager& Manager = UAssetManager::Get();
@@ -196,29 +195,27 @@ void UGameFeatureAction_AddChunkOverride::AddChunkIdOverride()
 void UGameFeatureAction_AddChunkOverride::RemoveChunkIdOverride()
 {
 #if WITH_EDITOR
-	if (LastChunkIdUsed < 0)
-	{
-		UE_LOG(LogAddChunkOverride, Verbose, TEXT("LastChunkIdUsed(%d) was invalid. Skipping override removal"), LastChunkIdUsed);
-		return;
-	}
-
 	// Remove primary asset rules by setting the override the default.
 	if (UGameFeatureData* GameFeatureData = GetTypedOuter<UGameFeatureData>())
 	{
 		FString PluginName;
 		GameFeatureData->GetPluginName(PluginName);
+		if (!GameFeatureAction_AddChunkOverride::PluginToChunkId.Contains(PluginName))
+		{
+			UE_LOG(LogAddChunkOverride, Verbose, TEXT("No chunk override found for (%s) Skipping override removal"), *PluginName);
+			return;
+		}
 
-		if (GameFeatureAction_AddChunkOverride::ChunkIdToPluginMap.Contains(LastChunkIdUsed))
+		const int32 ChunkIdOverride = GameFeatureAction_AddChunkOverride::PluginToChunkId[PluginName];
+		if (GameFeatureAction_AddChunkOverride::ChunkIdToPluginMap.Contains(ChunkIdOverride))
 		{
-			GameFeatureAction_AddChunkOverride::ChunkIdToPluginMap[LastChunkIdUsed].Remove(PluginName);
-			if (GameFeatureAction_AddChunkOverride::ChunkIdToPluginMap[LastChunkIdUsed].IsEmpty())
+			GameFeatureAction_AddChunkOverride::ChunkIdToPluginMap[ChunkIdOverride].Remove(PluginName);
+			if (GameFeatureAction_AddChunkOverride::ChunkIdToPluginMap[ChunkIdOverride].IsEmpty())
 			{
-				GameFeatureAction_AddChunkOverride::ChunkIdToPluginMap.Remove(LastChunkIdUsed);
+				GameFeatureAction_AddChunkOverride::ChunkIdToPluginMap.Remove(ChunkIdOverride);
 			}
 		}
-		ensure(GameFeatureAction_AddChunkOverride::PluginToChunkId.Remove(PluginName));
-		UE_LOG(LogAddChunkOverride, Log, TEXT("Removing ChunkId override (%d) for Plugin (%s)"), LastChunkIdUsed, *PluginName);
-		LastChunkIdUsed = -1;
+		UE_LOG(LogAddChunkOverride, Log, TEXT("Removing ChunkId override (%d) for Plugin (%s)"), ChunkIdOverride, *PluginName);
 
 		UAssetManager& Manager = UAssetManager::Get();
 
diff --git a/Engine/Plugins/Runtime/GameFeatures/Source/GameFeatures/Public/GameFeatureAction_AddChunkOverride.h b/Engine/Plugins/Runtime/GameFeatures/Source/GameFeatures/Public/GameFeatureAction_AddChunkOverride.h
--- a/Engine/Plugins/Runtime/GameFeatures/Source/GameFeatures/Public/GameFeatureAction_AddChunkOverride.h
+++ b/Engine/Plugins/Runtime/GameFeatures/Source/GameFeatures/Public/GameFeatureAction_AddChunkOverride.h
@@ -79,11 +79,6 @@ private:
 	UPROPERTY(EditAnywhere, Category = "Asset Management", meta=(EditCondition="bShouldOverrideChunk"))
 	int32 ParentChunk = 10;
 
-	/**
-	 * Internal tracking variable to remove the chunkId used when the ChunkId property changes.
-	 */
-	int32 LastChunkIdUsed = -1;
-
 	/**
 	 * Config defined value for what is the smallest chunk index the autogeneration code can generate.
 	 * If autogeneration produces a chunk index lower than this value users will need to manually define the chunk index this GFP will cook into.
