- Updated UGameFeatureActionConvertContentBundleWorldPartitionBuilder to support finding and updating an existing GameFeatureAction with an 'Add World Partition Content' action with no EDL asset.
- Added option to commandlet to remove the old 'Add World Partition Content (Content Bundle)' action using  -RemoveContentBundleAction (false by default)
- Fixed a regression caused by CL 32724869, ResetLoaders called on the package to delete before saving the new package was moving the actor back to the old package.

diff --git a/Engine/Plugins/Runtime/GameFeatures/Source/GameFeaturesEditor/Private/GameFeatureActionConvertContentBundleWorldPartitionBuilder.cpp b/Engine/Plugins/Runtime/GameFeatures/Source/GameFeaturesEditor/Private/GameFeatureActionConvertContentBundleWorldPartitionBuilder.cpp
--- a/Engine/Plugins/Runtime/GameFeatures/Source/GameFeaturesEditor/Private/GameFeatureActionConvertContentBundleWorldPartitionBuilder.cpp
+++ b/Engine/Plugins/Runtime/GameFeatures/Source/GameFeaturesEditor/Private/GameFeatureActionConvertContentBundleWorldPartitionBuilder.cpp
@@ -30,6 +30,7 @@ DEFINE_LOG_CATEGORY_STATIC(LogConvertContentBundleBuilder, All, All);
 UGameFeatureActionConvertContentBundleWorldPartitionBuilder::UGameFeatureActionConvertContentBundleWorldPartitionBuilder(const FObjectInitializer& ObjectInitializer)
 	: Super(ObjectInitializer)
 	, bReportOnly(false)
+	, bRemoveContentBundleAction(false)
 {}
 
 bool UGameFeatureActionConvertContentBundleWorldPartitionBuilder::PreRun(UWorld* World, FPackageSourceControlHelper& PackageHelper)
@@ -41,6 +42,7 @@ bool UGameFeatureActionConvertContentBundleWorldPartitionBuilder::PreRun(UWorld*
 	FString const* DestFolderPtr = CommandLineParams.Find(TEXT("DestinationFolder"));
 	DestinationFolder = DestFolderPtr ? *DestFolderPtr : FString();
 	bReportOnly = Switches.Contains(TEXT("ReportOnly"));
+	bRemoveContentBundleAction = Switches.Contains(TEXT("RemoveContentBundleAction"));
 
 	if (FString const* ContentBundlesToConvertString = CommandLineParams.Find(TEXT("ContentBundles")))
 	{
@@ -172,12 +174,25 @@ bool UGameFeatureActionConvertContentBundleWorldPartitionBuilder::RunInternal(UW
 		UGameFeatureAction_AddWorldPartitionContent* AddWorldPartitionContent = ExistingAddWorldPartitionContent ? Cast<UGameFeatureAction_AddWorldPartitionContent>(*ExistingAddWorldPartitionContent) : nullptr;
 		if (!AddWorldPartitionContent)
 		{
-			// Create new GameFeatureAction_AddWorldPartitionContent for this EDL Asset
-			AddWorldPartitionContent = NewObject<UGameFeatureAction_AddWorldPartitionContent>(GameFeatureData);
-			GameFeatureData->GetMutableActionsInEditor().Add(AddWorldPartitionContent);
-			AddWorldPartitionContent->ExternalDataLayerAsset = ExternalDataLayerAsset;
-			PackagesToSave.Add(GameFeatureData->GetPackage());
-			UE_LOG(LogConvertContentBundleBuilder, Log, TEXT("Added new Action of type 'GameFeatureAction_AddWorldPartitionContent' to GameFeatureData %s using External Data Layer Asset %s while converting Content Bundle %s."), *GameFeatureData->GetName(), *ExternalDataLayerAsset->GetPathName(), *ContentBundle->GetDisplayName());
+			// Find existing GameFeatureAction_AddWorldPartitionContent with no EDL asset
+			UGameFeatureAction* const* ExistingAddWorldPartitionContentWithNoEDLAsset = Algo::FindByPredicate(GameFeatureData->GetActions(), [ExternalDataLayerAsset](UGameFeatureAction* Action) { return Action && Action->IsA<UGameFeatureAction_AddWorldPartitionContent>() && Cast<UGameFeatureAction_AddWorldPartitionContent>(Action)->GetExternalDataLayerAsset() == nullptr; });
+			AddWorldPartitionContent = ExistingAddWorldPartitionContentWithNoEDLAsset ? Cast<UGameFeatureAction_AddWorldPartitionContent>(*ExistingAddWorldPartitionContentWithNoEDLAsset) : nullptr;
+			if (!AddWorldPartitionContent)
+			{
+				// Create new GameFeatureAction_AddWorldPartitionContent for this EDL Asset
+				AddWorldPartitionContent = NewObject<UGameFeatureAction_AddWorldPartitionContent>(GameFeatureData);
+				GameFeatureData->GetMutableActionsInEditor().Add(AddWorldPartitionContent);
+				PackagesToSave.Add(GameFeatureData->GetPackage());
+				UE_LOG(LogConvertContentBundleBuilder, Log, TEXT("Added new Action of type 'GameFeatureAction_AddWorldPartitionContent' to GameFeatureData %s using External Data Layer Asset %s while converting Content Bundle %s."), *GameFeatureData->GetName(), *ExternalDataLayerAsset->GetPathName(), *ContentBundle->GetDisplayName());
+			}
+
+			// Set EDL Asset on Action if necessary
+			if (AddWorldPartitionContent->ExternalDataLayerAsset != ExternalDataLayerAsset)
+			{
+				check(!AddWorldPartitionContent->ExternalDataLayerAsset);
+				AddWorldPartitionContent->ExternalDataLayerAsset = ExternalDataLayerAsset;
+				PackagesToSave.Add(GameFeatureData->GetPackage());
+			}
 		}
 		
 		// Backup the source ContentBundle Guid (to facilitate potential revert operation)
@@ -240,6 +255,7 @@ bool UGameFeatureActionConvertContentBundleWorldPartitionBuilder::RunInternal(UW
 			ActorReferences.Add(ActorRef);
 			AActor* Actor = ActorRef.GetActor();
 			UPackage* OldActorPackage = Actor->GetExternalPackage();
+			ResetLoaders(OldActorPackage);
 
 			FText FailureReason;
 			if (!FExternalDataLayerHelper::MoveActorsToExternalDataLayer({ Actor }, ExternalDataLayerInstance, &FailureReason))
@@ -268,7 +284,7 @@ bool UGameFeatureActionConvertContentBundleWorldPartitionBuilder::RunInternal(UW
 			ContentBundleContainerInstance->RemoveActor(ActorGuid);
 		}
 		
-		if (bContentBundleActorConversionSuccess)
+		if (bContentBundleActorConversionSuccess && bRemoveContentBundleAction)
 		{
 			// Remove Content Bundle GameFeatureData action
 			GameFeatureData->GetMutableActionsInEditor().Remove(OldAddWPContentAction);
@@ -290,7 +306,6 @@ bool UGameFeatureActionConvertContentBundleWorldPartitionBuilder::RunInternal(UW
 		for (UPackage* PackageToDelete : PackagesToDelete)
 		{
 			UE_LOG(LogConvertContentBundleBuilder, Log, TEXT("Package to delete: %s"), *PackageToDelete->GetPathName());
-			ResetLoaders(PackageToDelete);
 		}
 
 		if (bContentBundleActorConversionSuccess && !bReportOnly)
diff --git a/Engine/Plugins/Runtime/GameFeatures/Source/GameFeaturesEditor/Private/GameFeatureActionConvertContentBundleWorldPartitionBuilder.h b/Engine/Plugins/Runtime/GameFeatures/Source/GameFeaturesEditor/Private/GameFeatureActionConvertContentBundleWorldPartitionBuilder.h
--- a/Engine/Plugins/Runtime/GameFeatures/Source/GameFeaturesEditor/Private/GameFeatureActionConvertContentBundleWorldPartitionBuilder.h
+++ b/Engine/Plugins/Runtime/GameFeatures/Source/GameFeaturesEditor/Private/GameFeatureActionConvertContentBundleWorldPartitionBuilder.h
@@ -32,4 +32,5 @@ private:
 	TArray<FString> ContentBundlesToConvert;
 	FString DestinationFolder;
 	bool bReportOnly;
+	bool bRemoveContentBundleAction;
 };
\ No newline at end of file
