Add a scoped slow task to game feature data loading, adds information to the splash screen, informing user want they are waiting for
#rnx

diff --git a/Engine/Plugins/Runtime/GameFeatures/Source/GameFeatures/Private/GameFeaturePluginStateMachine.cpp b/Engine/Plugins/Runtime/GameFeatures/Source/GameFeatures/Private/GameFeaturePluginStateMachine.cpp
--- a/Engine/Plugins/Runtime/GameFeatures/Source/GameFeatures/Private/GameFeaturePluginStateMachine.cpp
+++ b/Engine/Plugins/Runtime/GameFeatures/Source/GameFeatures/Private/GameFeaturePluginStateMachine.cpp
@@ -20,6 +20,7 @@
 #include "Misc/EnumRange.h"
 #include "Misc/FileHelper.h"
 #include "Misc/PackageName.h"
+#include "Misc/ScopedSlowTask.h"
 #include "Misc/WildcardString.h"
 #include "Algo/AllOf.h"
 #include "Misc/TVariantMeta.h"
@@ -42,6 +43,8 @@
 
 #include UE_INLINE_GENERATED_CPP_BY_NAME(GameFeaturePluginStateMachine)
 
+#define LOCTEXT_NAMESPACE "GameFeatureDataStateMachine"
+
 #if WITH_EDITOR
 #include "PluginUtils.h"
 #include "Misc/App.h"
@@ -2633,6 +2636,11 @@ struct FGameFeaturePluginState_Registering : public FGameFeaturePluginState
 		}
 		else
 		{
+			FScopedSlowTask LoadingGameFeatureData(1.0f, 
+				FText::Format(
+					LOCTEXT("LoadingGameFeatureData", "Loading Game Feature Data for Plugin: {0}"), 
+					FText::FromString(StateProperties.PluginName)));
+			LoadingGameFeatureData.Visibility = ESlowTaskVisibility::Important;
 			for (const FString& Path : GameFeatureDataSearchPaths)
 			{
 				if (FPackageName::DoesPackageExist(Path))
@@ -3886,3 +3894,4 @@ bool FGameFeaturePluginStateMachineProperties::AllowAsyncLoading() const
 		(!IsRunningCommandlet() || UE::GameFeatures::CVarForceAsyncLoad.GetValueOnGameThread());
 }
 
+#undef LOCTEXT_NAMESPACE 
