Plugin templates can provide an optional default plugin name to use when selecting the template in the plugin wizard

diff --git a/Engine/Plugins/Editor/PluginBrowser/Source/PluginBrowser/Private/SNewPluginWizard.cpp b/Engine/Plugins/Editor/PluginBrowser/Source/PluginBrowser/Private/SNewPluginWizard.cpp
--- a/Engine/Plugins/Editor/PluginBrowser/Source/PluginBrowser/Private/SNewPluginWizard.cpp
+++ b/Engine/Plugins/Editor/PluginBrowser/Source/PluginBrowser/Private/SNewPluginWizard.cpp
@@ -444,6 +444,7 @@ void SNewPluginWizard::OnTemplateSelectionChanged(TSharedPtr<FPluginTemplateDesc
 			if (OldSelectedItem.IsValid())
 			{
 				OldSelectedItem->UpdatePathWhenTemplateUnselected(PluginFolderPath);
+				OldSelectedItem->UpdatePluginNameTextWhenTemplateUnselected(PluginNameText);
 			}
 
 			PluginWizardDefinition->OnTemplateSelectionChanged(SelectedItem, SelectInfo);
@@ -451,6 +452,7 @@ void SNewPluginWizard::OnTemplateSelectionChanged(TSharedPtr<FPluginTemplateDesc
 			if (SelectedItem.IsValid())
 			{
 				SelectedItem->UpdatePathWhenTemplateSelected(PluginFolderPath);
+				SelectedItem->UpdatePluginNameTextWhenTemplateSelected(PluginNameText);
 			}
 		}
 	}
diff --git a/Engine/Plugins/Runtime/GameFeatures/Source/GameFeaturesEditor/Private/GameFeaturePluginTemplate.cpp b/Engine/Plugins/Runtime/GameFeatures/Source/GameFeaturesEditor/Private/GameFeaturePluginTemplate.cpp
--- a/Engine/Plugins/Runtime/GameFeatures/Source/GameFeaturesEditor/Private/GameFeaturePluginTemplate.cpp
+++ b/Engine/Plugins/Runtime/GameFeatures/Source/GameFeaturesEditor/Private/GameFeaturePluginTemplate.cpp
@@ -16,12 +16,13 @@
 
 #define LOCTEXT_NAMESPACE "GameFeatures"
 
-FGameFeaturePluginTemplateDescription::FGameFeaturePluginTemplateDescription(FText InName, FText InDescription, FString InOnDiskPath, FString InDefaultSubfolder, TSubclassOf<UGameFeatureData> GameFeatureDataClassOverride, FString GameFeatureDataNameOverride, EPluginEnabledByDefault InEnabledByDefault)
+FGameFeaturePluginTemplateDescription::FGameFeaturePluginTemplateDescription(FText InName, FText InDescription, FString InOnDiskPath, FString InDefaultSubfolder, FString InDefaultPluginName, TSubclassOf<UGameFeatureData> GameFeatureDataClassOverride, FString GameFeatureDataNameOverride, EPluginEnabledByDefault InEnabledByDefault)
 	: FPluginTemplateDescription(InName, InDescription, InOnDiskPath, /*bCanContainContent=*/ true, EHostType::Runtime)
 {
 	SortPriority = 10;
 	bCanBePlacedInEngine = false;
 	DefaultSubfolder = InDefaultSubfolder;
+	DefaultPluginName = InDefaultPluginName;
 	GameFeatureDataName = !GameFeatureDataNameOverride.IsEmpty() ? GameFeatureDataNameOverride : FString();
 	GameFeatureDataClass = GameFeatureDataClassOverride != nullptr ? GameFeatureDataClassOverride : TSubclassOf<UGameFeatureData>(UGameFeatureData::StaticClass());
 	PluginEnabledByDefault = InEnabledByDefault;
@@ -53,6 +54,16 @@ void FGameFeaturePluginTemplateDescription::UpdatePathWhenTemplateUnselected(FSt
 	FPaths::MakePlatformFilename(InOutPath);
 }
 
+void FGameFeaturePluginTemplateDescription::UpdatePluginNameTextWhenTemplateSelected(FText& OutPluginNameText)
+{
+	OutPluginNameText = FText::FromString(DefaultPluginName);
+}
+
+void FGameFeaturePluginTemplateDescription::UpdatePluginNameTextWhenTemplateUnselected(FText& OutPluginNameText)
+{
+	OutPluginNameText = FText::GetEmpty();
+}
+
 void FGameFeaturePluginTemplateDescription::CustomizeDescriptorBeforeCreation(FPluginDescriptor& Descriptor)
 {
 	Descriptor.bExplicitlyLoaded = true;
diff --git a/Engine/Plugins/Runtime/GameFeatures/Source/GameFeaturesEditor/Private/GameFeaturesEditorModule.cpp b/Engine/Plugins/Runtime/GameFeatures/Source/GameFeaturesEditor/Private/GameFeaturesEditorModule.cpp
--- a/Engine/Plugins/Runtime/GameFeatures/Source/GameFeaturesEditor/Private/GameFeaturesEditorModule.cpp
+++ b/Engine/Plugins/Runtime/GameFeatures/Source/GameFeaturesEditor/Private/GameFeaturesEditorModule.cpp
@@ -112,6 +112,7 @@ class FGameFeaturesEditorModule : public FDefaultModuleImpl
 					PluginTemplate.Description,
 					PluginTemplate.Path.Path,
 					PluginTemplate.DefaultSubfolder,
+					PluginTemplate.DefaultPluginName,
 					PluginTemplate.DefaultGameFeatureDataClass,
 					PluginTemplate.DefaultGameFeatureDataName,
 					PluginTemplate.bIsEnabledByDefault ? EPluginEnabledByDefault::Enabled : EPluginEnabledByDefault::Disabled)));
diff --git a/Engine/Plugins/Runtime/GameFeatures/Source/GameFeaturesEditor/Private/GameFeaturesEditorSettings.h b/Engine/Plugins/Runtime/GameFeatures/Source/GameFeaturesEditor/Private/GameFeaturesEditorSettings.h
--- a/Engine/Plugins/Runtime/GameFeatures/Source/GameFeaturesEditor/Private/GameFeaturesEditorSettings.h
+++ b/Engine/Plugins/Runtime/GameFeatures/Source/GameFeaturesEditor/Private/GameFeaturesEditorSettings.h
@@ -34,6 +34,10 @@ struct FPluginTemplateData
 	UPROPERTY(EditAnywhere, Category = PluginTemplate)
 	FString DefaultSubfolder;
 
+	/** Optional plugin name to default the new plugin to. */
+	UPROPERTY(EditAnywhere, Category = PluginTemplate)
+	FString DefaultPluginName;
+
 	/** The default class of game feature data to create for new game feature plugins (if not set, UGameFeatureData will be used) */
 	UPROPERTY(config, EditAnywhere, Category = Plugins)
 	TSubclassOf<UGameFeatureData> DefaultGameFeatureDataClass;
diff --git a/Engine/Plugins/Runtime/GameFeatures/Source/GameFeaturesEditor/Public/GameFeaturePluginTemplate.h b/Engine/Plugins/Runtime/GameFeatures/Source/GameFeaturesEditor/Public/GameFeaturePluginTemplate.h
--- a/Engine/Plugins/Runtime/GameFeatures/Source/GameFeaturesEditor/Public/GameFeaturePluginTemplate.h
+++ b/Engine/Plugins/Runtime/GameFeatures/Source/GameFeaturesEditor/Public/GameFeaturePluginTemplate.h
@@ -11,13 +11,16 @@
  */
 struct GAMEFEATURESEDITOR_API FGameFeaturePluginTemplateDescription : public FPluginTemplateDescription
 {
-	FGameFeaturePluginTemplateDescription(FText InName, FText InDescription, FString InOnDiskPath, FString InDefaultSubfolder
+	FGameFeaturePluginTemplateDescription(FText InName, FText InDescription, FString InOnDiskPath, FString InDefaultSubfolder, FString InDefaultPluginName
 		, TSubclassOf<UGameFeatureData> GameFeatureDataClassOverride, FString GameFeatureDataNameOverride, EPluginEnabledByDefault InEnabledByDefault);
 
 	virtual bool ValidatePathForPlugin(const FString& ProposedAbsolutePluginPath, FText& OutErrorMessage) override;
 	virtual void UpdatePathWhenTemplateSelected(FString& InOutPath) override;
 	virtual void UpdatePathWhenTemplateUnselected(FString& InOutPath) override;
 
+	virtual void UpdatePluginNameTextWhenTemplateSelected(FText& OutPluginNameText) override;
+	virtual void UpdatePluginNameTextWhenTemplateUnselected(FText& OutPluginNameText) override;
+
 	virtual void CustomizeDescriptorBeforeCreation(FPluginDescriptor& Descriptor) override;
 	virtual void OnPluginCreated(TSharedPtr<IPlugin> NewPlugin) override;
 
@@ -25,6 +28,7 @@ struct GAMEFEATURESEDITOR_API FGameFeaturePluginTemplateDescription : public FPl
 	bool IsRootedInGameFeaturesRoot(const FString& InStr) const;
 
 	FString DefaultSubfolder;
+	FString DefaultPluginName;
 	TSubclassOf<UGameFeatureData> GameFeatureDataClass;
 	FString GameFeatureDataName;
 	EPluginEnabledByDefault PluginEnabledByDefault = EPluginEnabledByDefault::Disabled;
diff --git a/Engine/Source/Editor/UnrealEd/Public/Features/IPluginsEditorFeature.h b/Engine/Source/Editor/UnrealEd/Public/Features/IPluginsEditorFeature.h
--- a/Engine/Source/Editor/UnrealEd/Public/Features/IPluginsEditorFeature.h
+++ b/Engine/Source/Editor/UnrealEd/Public/Features/IPluginsEditorFeature.h
@@ -63,6 +63,12 @@ struct FPluginTemplateDescription
 	/** Called to change away from special folders if needed */
 	virtual void UpdatePathWhenTemplateUnselected(FString& InOutPath) {}
 
+	/** Called to optionally set the default plugin name text */
+	virtual void UpdatePluginNameTextWhenTemplateSelected(FText& OutPluginNameText) {}
+
+	/** Called to optionally unset the default plugin name text */
+	virtual void UpdatePluginNameTextWhenTemplateUnselected(FText& OutPluginNameText) {}
+
 	/** Constructor */
 	FPluginTemplateDescription(FText InName, FText InDescription, FString InOnDiskPath, bool InCanContainContent, EHostType::Type InModuleDescriptorType, ELoadingPhase::Type InLoadingPhase = ELoadingPhase::Default)
 		: Name(InName)
