Fix issue where variable was initialized through itself.
Caught by PVS

diff --git a/Engine/Plugins/Runtime/GameFeatures/Source/GameFeatures/Private/GameFeaturePluginStateMachine.cpp b/Engine/Plugins/Runtime/GameFeatures/Source/GameFeatures/Private/GameFeaturePluginStateMachine.cpp
--- a/Engine/Plugins/Runtime/GameFeatures/Source/GameFeatures/Private/GameFeaturePluginStateMachine.cpp
+++ b/Engine/Plugins/Runtime/GameFeatures/Source/GameFeatures/Private/GameFeaturePluginStateMachine.cpp
@@ -2397,7 +2397,7 @@ struct FGameFeaturePluginState_AssetDependencyStreaming : public FGameFeaturePlu
 
 		TSharedPtr<IInstallBundleManager> BundleManager = IInstallBundleManager::GetPlatformInstallBundleManager();
 
-		EInstallBundleRequestFlags InstallFlags = InstallFlags = UseAsyncLoading() ?
+		EInstallBundleRequestFlags InstallFlags = UseAsyncLoading() ?
 			(Options.InstallBundleFlags | EInstallBundleRequestFlags::AsyncMount) : Options.InstallBundleFlags;
 		TValueOrError<FInstallBundleRequestInfo, EInstallBundleResult> MaybeRequestInfo = BundleManager->RequestUpdateContent(AssetInstallBundles, InstallFlags);
 		if (MaybeRequestInfo.HasError())
