Fully implement all GFP states and functionality for asset dependency streaming.

diff --git a/Engine/Plugins/Runtime/GameFeatures/Source/GameFeatures/Private/GameFeaturePluginStateMachine.cpp b/Engine/Plugins/Runtime/GameFeatures/Source/GameFeatures/Private/GameFeaturePluginStateMachine.cpp
--- a/Engine/Plugins/Runtime/GameFeatures/Source/GameFeatures/Private/GameFeaturePluginStateMachine.cpp
+++ b/Engine/Plugins/Runtime/GameFeatures/Source/GameFeatures/Private/GameFeaturePluginStateMachine.cpp
@@ -775,10 +775,11 @@ struct FGameFeaturePluginState_CheckingStatus : public FGameFeaturePluginState
 	{
 		if (!bParsedURL)
 		{
-			bParsedURL = StateProperties.ParseURL();
+			TValueOrError<void, FString> ParseUrlResult = StateProperties.ParseURL();
+			bParsedURL = !ParseUrlResult.HasError();
 			if (!bParsedURL)
 			{
-				StateStatus.SetTransitionError(EGameFeaturePluginState::ErrorCheckingStatus, GetErrorResult(TEXT("Bad_PluginURL")));
+				StateStatus.SetTransitionError(EGameFeaturePluginState::ErrorCheckingStatus, GetErrorResult(ParseUrlResult.GetError()));
 				return;
 			}
 		}
@@ -959,6 +960,12 @@ struct FBaseDataReleaseGameFeaturePluginState : public FGameFeaturePluginState
 	bool bWasDeleted = false;
 	TArray<FName> PendingBundles;
 
+	void CleanUp()
+	{
+		PendingBundles.Empty();
+		IInstallBundleManager::ReleasedDelegate.RemoveAll(this);
+	}
+
 	void OnContentRemoved(FInstallBundleReleaseRequestResultInfo BundleResult)
 	{
 		if (!PendingBundles.Contains(BundleResult.BundleName))
@@ -986,12 +993,7 @@ struct FBaseDataReleaseGameFeaturePluginState : public FGameFeaturePluginState
 		UpdateStateMachineImmediate();
 	}
 
-	virtual void BeginState() override
-	{
-		BeginRemoveRequest();
-	}
-
-	virtual void BeginRemoveRequest()
+	void BeginRemoveRequest()
 	{
 		CleanUp();
 
@@ -1007,10 +1009,8 @@ struct FBaseDataReleaseGameFeaturePluginState : public FGameFeaturePluginState
 		TSharedPtr<IInstallBundleManager> BundleManager = IInstallBundleManager::GetPlatformInstallBundleManager();
 		check(BundleManager.IsValid());
 
-		const TArray<FName>& InstallBundles = StateProperties.ProtocolMetadata.GetSubtype<FInstallBundlePluginProtocolMetaData>().InstallBundles;
-
 		EInstallBundleReleaseRequestFlags ReleaseFlags = GetReleaseRequestFlags();
-		TValueOrError<FInstallBundleReleaseRequestInfo, EInstallBundleResult> MaybeRequestInfo = BundleManager->RequestReleaseContent(InstallBundles, ReleaseFlags);
+		TValueOrError<FInstallBundleReleaseRequestInfo, EInstallBundleResult> MaybeRequestInfo = BundleManager->RequestReleaseContent(GetInstallBundles(), ReleaseFlags);
 
 		if (MaybeRequestInfo.HasError())
 		{
@@ -1043,6 +1043,11 @@ struct FBaseDataReleaseGameFeaturePluginState : public FGameFeaturePluginState
 		}
 	}
 
+	virtual void BeginState() override
+	{
+		BeginRemoveRequest();
+	}
+
 	virtual void UpdateState(FGameFeaturePluginStateStatus& StateStatus) override
 	{
 		if (!Result.HasValue())
@@ -1059,19 +1064,13 @@ struct FBaseDataReleaseGameFeaturePluginState : public FGameFeaturePluginState
 		StateStatus.SetTransition(GetSuccessTransitionState());
 	}
 
-	void CleanUp()
-	{
-		PendingBundles.Empty();
-		IInstallBundleManager::ReleasedDelegate.RemoveAll(this);
-	}
-
 	virtual void EndState() override
 	{
 		CleanUp();
 	}
 
 	/** Controls what check is done to determine if this state should run or not */
-	virtual bool ShouldReleaseContent() const
+	bool ShouldReleaseContent() const
 	{
 		switch (StateProperties.GetPluginProtocol())
 		{
@@ -1086,6 +1085,17 @@ struct FBaseDataReleaseGameFeaturePluginState : public FGameFeaturePluginState
 			}
 		}
 	}
+
+	virtual TConstArrayView<FName> GetInstallBundles()
+	{
+		TConstArrayView<FName> Ret;
+		if (ShouldReleaseContent())
+		{
+			Ret = StateProperties.ProtocolMetadata.GetSubtype<FInstallBundlePluginProtocolMetaData>().InstallBundles;
+		}
+		return Ret;
+	}
+
 	/** Determine what kind of release request flags we submit */
 	virtual EInstallBundleReleaseRequestFlags GetReleaseRequestFlags() const
 	{
@@ -1203,14 +1213,14 @@ struct FGameFeaturePluginState_Releasing : public FBaseDataReleaseGameFeaturePlu
 	}
 };
 
-struct FGameFeaturePluginState_Downloading : public FGameFeaturePluginState
+struct FBaseDownloadGameFeaturePluginState : public FGameFeaturePluginState
 {
-	FGameFeaturePluginState_Downloading(FGameFeaturePluginStateMachineProperties& InStateProperties)
+	FBaseDownloadGameFeaturePluginState(FGameFeaturePluginStateMachineProperties& InStateProperties)
 		: FGameFeaturePluginState(InStateProperties)
 		, Result(MakeValue())
 	{}
 
-	~FGameFeaturePluginState_Downloading()
+	virtual ~FBaseDownloadGameFeaturePluginState()
 	{
 		Cleanup();
 	}
@@ -1261,10 +1271,7 @@ struct FGameFeaturePluginState_Downloading : public FGameFeaturePluginState
 			UpdateStateMachineImmediate();
 			return;
 		}
-		
-		FInstallBundlePluginProtocolMetaData& Metadata = StateProperties.ProtocolMetadata.GetSubtype<FInstallBundlePluginProtocolMetaData>();
-		const TArray<FName>& InstallBundles = Metadata.InstallBundles;
-		
+
 		// check to verify if the bundle(s) we need is already UpToDate
 		if (BundleContentState.GetAllBundlesHaveState(EInstallBundleInstallState::UpToDate))
 		{
@@ -1284,11 +1291,8 @@ struct FGameFeaturePluginState_Downloading : public FGameFeaturePluginState
 			return;
 		}
 
-		//Pull our InstallFlags from the Options, but also make sure SkipMount is set as there is a separate mounting step that will re-request this
-		//without SkipMount and then mount the data, this allows us to pre-download data without mounting it
-		EInstallBundleRequestFlags InstallFlags = Options.InstallBundleFlags;
-		InstallFlags |= EInstallBundleRequestFlags::SkipMount;
-
+		const TConstArrayView<FName> InstallBundles = GetInstallBundles();
+		const EInstallBundleRequestFlags InstallFlags = GetRequestFlags();
 		TValueOrError<FInstallBundleRequestInfo, EInstallBundleResult> MaybeRequestInfo = BundleManager->RequestUpdateContent(InstallBundles, InstallFlags);
 
 		if (MaybeRequestInfo.HasError())
@@ -1307,7 +1311,7 @@ struct FGameFeaturePluginState_Downloading : public FGameFeaturePluginState
 			const FStringView ShortUrl = StateProperties.PluginIdentifier.GetIdentifyingString();
 			ensureMsgf(false, TEXT("Unable to enqueue download for the PluginURL(%.*s) because failed to resolve install bundles!"), ShortUrl.Len(), ShortUrl.GetData());
 			Result = GetErrorResult(TEXT("BundleManager.GotState."), TEXT("Resolve_Failed"), UE::GameFeatures::CommonErrorCodes::GetGenericConnectionError());
-			
+
 			UpdateStateMachineImmediate();
 			return;
 		}
@@ -1321,14 +1325,14 @@ struct FGameFeaturePluginState_Downloading : public FGameFeaturePluginState
 		else
 		{
 			PendingBundleDownloads = MoveTemp(RequestInfo.BundlesEnqueued);
-			IInstallBundleManager::InstallBundleCompleteDelegate.AddRaw(this, &FGameFeaturePluginState_Downloading::OnInstallBundleCompleted);
-			IInstallBundleManager::PausedBundleDelegate.AddRaw(this, &FGameFeaturePluginState_Downloading::OnInstallBundlePaused);
+			IInstallBundleManager::InstallBundleCompleteDelegate.AddRaw(this, &FBaseDownloadGameFeaturePluginState::OnInstallBundleCompleted);
+			IInstallBundleManager::PausedBundleDelegate.AddRaw(this, &FBaseDownloadGameFeaturePluginState::OnInstallBundlePaused);
 
 			ProgressTracker = MakeUnique<FInstallBundleCombinedProgressTracker>(false);
 			ProgressTracker->SetBundlesToTrackFromContentState(BundleContentState, PendingBundleDownloads);
 
 			ProgressUpdateHandle = FTSTicker::GetCoreTicker().AddTicker(
-				FTickerDelegate::CreateRaw(this, &FGameFeaturePluginState_Downloading::OnUpdateProgress)/*, 0.1f*/);
+				FTickerDelegate::CreateRaw(this, &FBaseDownloadGameFeaturePluginState::OnUpdateProgress)/*, 0.1f*/);
 
 			//If this setting is flipped then we should immediately request to pause downloads.
 			//We still generate the downloads so that we have an accurate PendingBundleDownloads list
@@ -1353,10 +1357,10 @@ struct FGameFeaturePluginState_Downloading : public FGameFeaturePluginState
 			//Use OptionalErrorCode and/or OptionalErrorText if available
 			const FString ErrorCodeEnding = (BundleResult.OptionalErrorCode.IsEmpty()) ? LexToString(BundleResult.Result) : BundleResult.OptionalErrorCode;
 			const FText ErrorText = BundleResult.OptionalErrorCode.IsEmpty() ? UE::GameFeatures::CommonErrorCodes::GetErrorTextForBundleResult(BundleResult.Result) : BundleResult.OptionalErrorText;
-			
+
 			Result = GetErrorResult(TEXT("BundleManager.OnComplete."), ErrorCodeEnding, ErrorText);
-			
-			if(BundleResult.Result != EInstallBundleResult::UserCancelledError)
+
+			if (BundleResult.Result != EInstallBundleResult::UserCancelledError)
 			{
 				TSharedPtr<IInstallBundleManager> BundleManager = IInstallBundleManager::GetPlatformInstallBundleManager();
 				BundleManager->CancelUpdateContent(PendingBundleDownloads);
@@ -1423,7 +1427,7 @@ struct FGameFeaturePluginState_Downloading : public FGameFeaturePluginState
 		{
 			const bool bIsPaused = (InPauseBundleInfo.PauseFlags != EInstallBundlePauseFlags::None);
 			const TCHAR* PauseReason = InstallBundleUtil::GetInstallBundlePauseReason(InPauseBundleInfo.PauseFlags);
-			
+
 			NotifyPauseChange(bIsPaused, PauseReason);
 		}
 	}
@@ -1438,18 +1442,22 @@ struct FGameFeaturePluginState_Downloading : public FGameFeaturePluginState
 	{
 		Cleanup();
 
-		check(StateProperties.GetPluginProtocol() == EGameFeaturePluginProtocol::InstallBundle);
-		ensureMsgf(AllowAsyncLoading(), TEXT("FGameFeaturePluginState::AllowAsyncLoading is false while attempting to download GFP data."));
+		if (!ShouldDownloadContent())
+		{
+			bPluginDownloaded = true;
+			UpdateProgress(1.0f);
+			return;
+		}
 
-		UGameFeaturesSubsystem::Get().OnGameFeatureDownloading(StateProperties.PluginName, StateProperties.PluginIdentifier);
+		ensureMsgf(AllowAsyncLoading(), TEXT("FGameFeaturePluginState::AllowAsyncLoading is false while attempting to download GFP data."));
 
 		TSharedPtr<IInstallBundleManager> BundleManager = IInstallBundleManager::GetPlatformInstallBundleManager();
-		const TArray<FName>& InstallBundles = StateProperties.ProtocolMetadata.GetSubtype<FInstallBundlePluginProtocolMetaData>().InstallBundles;
+		const TConstArrayView<FName> InstallBundles = GetInstallBundles();
 
 		if (InstallBundles.Num() > 1)
 		{
-			GotContentStateHandle = BundleManager->GetContentState(InstallBundles, EInstallBundleGetContentStateFlags::None, true, 
-				FInstallBundleGetContentStateDelegate::CreateRaw(this, &FGameFeaturePluginState_Downloading::OnGotContentState));
+			GotContentStateHandle = BundleManager->GetContentState(InstallBundles, EInstallBundleGetContentStateFlags::None, true,
+				FInstallBundleGetContentStateDelegate::CreateRaw(this, &FBaseDownloadGameFeaturePluginState::OnGotContentState));
 		}
 		else
 		{
@@ -1461,7 +1469,7 @@ struct FGameFeaturePluginState_Downloading : public FGameFeaturePluginState
 			const FInstallBundleCombinedInstallState& InstallState = MaybeInstallState.GetValue();
 			FInstallBundleCombinedContentState HackContentState;
 			HackContentState.IndividualBundleStates.Reserve(InstallState.IndividualBundleStates.Num());
-			for(const TPair<FName, EInstallBundleInstallState>& Pair : InstallState.IndividualBundleStates)
+			for (const TPair<FName, EInstallBundleInstallState>& Pair : InstallState.IndividualBundleStates)
 			{
 				FInstallBundleContentState& BundleContentState = HackContentState.IndividualBundleStates.Emplace(Pair.Key);
 				BundleContentState.State = Pair.Value;
@@ -1475,7 +1483,8 @@ struct FGameFeaturePluginState_Downloading : public FGameFeaturePluginState
 	{
 		if (!Result.HasValue())
 		{
-			StateStatus.SetTransitionError(EGameFeaturePluginState::ErrorManagingData, Result);
+			const EGameFeaturePluginState FailState = GetFailureTransitionState();
+			StateStatus.SetTransitionError(FailState, Result);
 			return;
 		}
 
@@ -1484,7 +1493,8 @@ struct FGameFeaturePluginState_Downloading : public FGameFeaturePluginState
 			return;
 		}
 
-		StateStatus.SetTransition(EGameFeaturePluginState::Installed);
+		const EGameFeaturePluginState SuccessState = GetSuccessTransitionState();
+		StateStatus.SetTransition(SuccessState);
 	}
 
 	virtual void TryCancelState() override
@@ -1520,11 +1530,11 @@ struct FGameFeaturePluginState_Downloading : public FGameFeaturePluginState
 		}
 
 		const FInstallBundlePluginProtocolOptions& Options = StateProperties.ProtocolOptions.GetSubtype<FInstallBundlePluginProtocolOptions>();
-		
+
 		//Update our InstallBundleRequestFlags
 		{
 			EInstallBundleRequestFlags UpdatedRequestFlags = Options.InstallBundleFlags;
-						
+
 			EInstallBundleRequestFlags AddFlags = (UpdatedRequestFlags & (~OldRequestFlags));
 			EInstallBundleRequestFlags RemoveFlags = ((~UpdatedRequestFlags) & OldRequestFlags);
 
@@ -1550,6 +1560,74 @@ struct FGameFeaturePluginState_Downloading : public FGameFeaturePluginState
 	{
 		Cleanup();
 	}
+
+	/** Controls what check is done to determine if this state should run or not */
+	bool ShouldDownloadContent() const
+	{
+		switch (StateProperties.GetPluginProtocol())
+		{
+		case (EGameFeaturePluginProtocol::InstallBundle):
+			return true;
+		default:
+			return false;
+		}
+	}
+
+	virtual TConstArrayView<FName> GetInstallBundles()
+	{
+		TConstArrayView<FName> Ret;
+		if (ShouldDownloadContent())
+		{
+			Ret = StateProperties.ProtocolMetadata.GetSubtype<FInstallBundlePluginProtocolMetaData>().InstallBundles;
+		}
+		return Ret;
+	}
+
+	/** Determine what kind of request flags we submit */
+	virtual EInstallBundleRequestFlags GetRequestFlags() const
+	{
+		EInstallBundleRequestFlags InstallFlags = StateProperties.ProtocolOptions.GetSubtype<FInstallBundlePluginProtocolOptions>().InstallBundleFlags;
+		return InstallFlags;
+	}
+
+	/** Determines what state you transition to in the event of a success or failure to release content */
+	virtual EGameFeaturePluginState GetSuccessTransitionState() const = 0;
+	virtual EGameFeaturePluginState GetFailureTransitionState() const = 0;
+};
+
+struct FGameFeaturePluginState_Downloading : public FBaseDownloadGameFeaturePluginState
+{
+	FGameFeaturePluginState_Downloading(FGameFeaturePluginStateMachineProperties& InStateProperties)
+		: FBaseDownloadGameFeaturePluginState(InStateProperties)
+	{}
+
+	virtual void BeginState() override
+	{
+		check(ShouldDownloadContent());
+
+		UGameFeaturesSubsystem::Get().OnGameFeatureDownloading(StateProperties.PluginName, StateProperties.PluginIdentifier);
+
+		FBaseDownloadGameFeaturePluginState::BeginState();
+	}
+	
+	virtual EInstallBundleRequestFlags GetRequestFlags() const override
+	{
+		//Pull our InstallFlags from the Options, but also make sure SkipMount is set as there is a separate mounting step that will re-request this
+		//without SkipMount and then mount the data, this allows us to pre-download data without mounting it
+		EInstallBundleRequestFlags InstallFlags = FBaseDownloadGameFeaturePluginState::GetRequestFlags();
+		InstallFlags |= EInstallBundleRequestFlags::SkipMount;
+		return InstallFlags;
+	}
+
+	virtual EGameFeaturePluginState GetSuccessTransitionState() const override
+	{
+		return EGameFeaturePluginState::Installed;
+	}
+
+	virtual EGameFeaturePluginState GetFailureTransitionState() const override
+	{
+		return EGameFeaturePluginState::ErrorManagingData;
+	}
 };
 
 struct FGameFeaturePluginState_Installed : public FDestinationGameFeaturePluginState
@@ -2296,157 +2374,95 @@ struct FGameFeaturePluginState_WaitingForDependencies : public FTransitionDepend
 	}
 };
 
-// TODO: this is mostly going to look like the download state with a few modifications. Should probably just have a
-// common base with it
-struct FGameFeaturePluginState_AssetDependencyStreaming : public FGameFeaturePluginState
+struct FGameFeaturePluginState_AssetDependencyStreamOut : public FBaseDataReleaseGameFeaturePluginState
 {
-	FGameFeaturePluginState_AssetDependencyStreaming(FGameFeaturePluginStateMachineProperties& InStateProperties) 
-		: FGameFeaturePluginState(InStateProperties) 
-		, Result(MakeValue())
+	FGameFeaturePluginState_AssetDependencyStreamOut(FGameFeaturePluginStateMachineProperties& InStateProperties)
+		: FBaseDataReleaseGameFeaturePluginState(InStateProperties)
 	{}
 
-	~FGameFeaturePluginState_AssetDependencyStreaming()
+	virtual EGameFeaturePluginState GetSuccessTransitionState() const override
 	{
-		Cleanup();
+		return EGameFeaturePluginState::Unmounting;
 	}
 
-	void Cleanup()
+	virtual EGameFeaturePluginState GetFailureTransitionState() const override
 	{
-		IInstallBundleManager::InstallBundleCompleteDelegate.RemoveAll(this);
-		Result = MakeValue();
-		PendingBundleDownloads.Empty();
+		return EGameFeaturePluginState::ErrorAssetDependencyStreaming;
 	}
 
-	UE::GameFeatures::FResult Result;
-	TArray<FName> PendingBundleDownloads;
-
-	// TODO: How can we handle progress for this, we don't know we are going to do it until way after downloading
-	// TUniquePtr<FInstallBundleCombinedProgressTracker> ProgressTracker;
-
-	void OnInstallBundleCompleted(FInstallBundleRequestResultInfo BundleResult)
+	virtual TConstArrayView<FName> GetInstallBundles() override
 	{
-		if (!PendingBundleDownloads.Contains(BundleResult.BundleName))
-		{
-			return;
-		}
-
-		PendingBundleDownloads.Remove(BundleResult.BundleName);
-
-		if (!Result.HasError() && BundleResult.Result != EInstallBundleResult::OK)
-		{
-			//Use OptionalErrorCode and/or OptionalErrorText if available
-			const FString ErrorCodeEnding = (BundleResult.OptionalErrorCode.IsEmpty()) ? LexToString(BundleResult.Result) : BundleResult.OptionalErrorCode;
-			const FText ErrorText = BundleResult.OptionalErrorCode.IsEmpty() ? UE::GameFeatures::CommonErrorCodes::GetErrorTextForBundleResult(BundleResult.Result) : BundleResult.OptionalErrorText;
-
-			Result = GetErrorResult(TEXT("BundleManager."), ErrorCodeEnding, ErrorText);
-
-			if (BundleResult.Result != EInstallBundleResult::UserCancelledError)
-			{
-				TSharedPtr<IInstallBundleManager> BundleManager = IInstallBundleManager::GetPlatformInstallBundleManager();
-				BundleManager->CancelUpdateContent(PendingBundleDownloads);
-			}
-		}
-
-		if (PendingBundleDownloads.Num() > 0)
+		TConstArrayView<FName> Ret;
+		if (ShouldReleaseContent())
 		{
-			return;
+			Ret = StateProperties.ProtocolMetadata.GetSubtype<FInstallBundlePluginProtocolMetaData>().AssetDependencyBundles;
 		}
-
-		UpdateStateMachineImmediate();
+		return Ret;
 	}
 
-	virtual void BeginState() override
+	// There's no explict uninstall set for asset dependencies, so set EInstallBundleReleaseRequestFlags::RemoveFilesIfPossible
+	// to be 'correct'. In most cases, we'd expect to store these with a cache of some sort, so this exact flag probably wont matter
+	// in practice.
+	virtual EInstallBundleReleaseRequestFlags GetReleaseRequestFlags() const override
 	{
-		Cleanup();
-
-		if (StateProperties.GetPluginProtocol() != EGameFeaturePluginProtocol::InstallBundle)
-		{
-			// For now, we don't support asset streaming with file protocol
-			return;
-		}
-
-		// TODO: Install Bundles need to move away from FNames for identifiers, this is currently just bloating up the name table
-		// when using dynamic GFPs
-		TValueOrError<TArray<FName>, FString> MaybeAssetInstallBundles = 
-			UGameFeaturesSubsystem::Get().GetPolicy().GetStreamingAssetInstallBundles(StateProperties.PluginIdentifier.GetFullPluginURL());
-		if (MaybeAssetInstallBundles.HasError())
-		{
-			Result = GetErrorResult(MaybeAssetInstallBundles.GetError());
-			return;
-		}
+		const EInstallBundleReleaseRequestFlags BaseFlags = FBaseDataReleaseGameFeaturePluginState::GetReleaseRequestFlags();
+		return (BaseFlags | EInstallBundleReleaseRequestFlags::RemoveFilesIfPossible);
+	}
+};
 
-		const TArray<FName>& AssetInstallBundles = MaybeAssetInstallBundles.GetValue();
-		if (AssetInstallBundles.IsEmpty())
-		{
-			return;
-		}
+struct FGameFeaturePluginState_ErrorAssetDependencyStreaming : public FErrorGameFeaturePluginState
+{
+	FGameFeaturePluginState_ErrorAssetDependencyStreaming(FGameFeaturePluginStateMachineProperties& InStateProperties) 
+		: FErrorGameFeaturePluginState(InStateProperties) 
+	{}
 
-		const FInstallBundlePluginProtocolOptions& Options = StateProperties.ProtocolOptions.GetSubtype<FInstallBundlePluginProtocolOptions>();
-		if (Options.bDoNotDownload)
+	virtual void UpdateState(FGameFeaturePluginStateStatus& StateStatus) override
+	{
+		if (StateProperties.Destination < EGameFeaturePluginState::ErrorAssetDependencyStreaming)
 		{
-			Result = GetErrorResult(TEXT("GFPStateMachine.DownloadNotAllowed"));
-			return;
+			StateStatus.SetTransition(EGameFeaturePluginState::AssetDependencyStreamOut);
 		}
-
-		FInstallBundlePluginProtocolMetaData& ProtocolData = StateProperties.ProtocolMetadata.GetSubtype<FInstallBundlePluginProtocolMetaData>();
-		for (FName Bundle : AssetInstallBundles)
+		else if (StateProperties.Destination > EGameFeaturePluginState::ErrorAssetDependencyStreaming)
 		{
-			// Add to metadata list for later cleanup
-			ProtocolData.InstallBundles.AddUnique(Bundle);
+			StateStatus.SetTransition(EGameFeaturePluginState::AssetDependencyStreaming);
 		}
+	}
+};
 
-		TSharedPtr<IInstallBundleManager> BundleManager = IInstallBundleManager::GetPlatformInstallBundleManager();
-
-		EInstallBundleRequestFlags InstallFlags = UseAsyncLoading() ?
-			(Options.InstallBundleFlags | EInstallBundleRequestFlags::AsyncMount) : Options.InstallBundleFlags;
-		TValueOrError<FInstallBundleRequestInfo, EInstallBundleResult> MaybeRequestInfo = BundleManager->RequestUpdateContent(AssetInstallBundles, InstallFlags);
-		if (MaybeRequestInfo.HasError())
-		{
-			const FStringView ShortUrl = StateProperties.PluginIdentifier.GetIdentifyingString();
-			ensureMsgf(false, TEXT("Unable to enqueue asset dependencies for the PluginURL(%.*s) because %s"), ShortUrl.Len(), ShortUrl.GetData(), LexToString(MaybeRequestInfo.GetError()));
-			Result = GetErrorResult(TEXT("BundleManager."), LexToString(MaybeRequestInfo.GetError()));
-			return;
-		}
+struct FGameFeaturePluginState_AssetDependencyStreaming : public FBaseDownloadGameFeaturePluginState
+{
+	FGameFeaturePluginState_AssetDependencyStreaming(FGameFeaturePluginStateMachineProperties& InStateProperties)
+		: FBaseDownloadGameFeaturePluginState(InStateProperties)
+	{}
 
-		FInstallBundleRequestInfo RequestInfo = MaybeRequestInfo.StealValue();
-		PendingBundleDownloads = MoveTemp(RequestInfo.BundlesEnqueued);
-		if (!PendingBundleDownloads.IsEmpty())
+	virtual TConstArrayView<FName> GetInstallBundles() override
+	{
+		TConstArrayView<FName> Ret;
+		if (ShouldDownloadContent())
 		{
-			IInstallBundleManager::InstallBundleCompleteDelegate.AddRaw(this, &FGameFeaturePluginState_AssetDependencyStreaming::OnInstallBundleCompleted);
+			Ret = StateProperties.ProtocolMetadata.GetSubtype<FInstallBundlePluginProtocolMetaData>().AssetDependencyBundles;
 		}
-
-		// TODO: how to apply pausing? bUserPauseDownload is protocol specific
-		//if (Options.bUserPauseDownload)
+		return Ret;
 	}
 
-	virtual void UpdateState(FGameFeaturePluginStateStatus& StateStatus) override
+	virtual EInstallBundleRequestFlags GetRequestFlags() const override
 	{
-		if (!Result.HasValue())
-		{
-			StateStatus.SetTransitionError(EGameFeaturePluginState::ErrorWaitingForDependencies, Result);
-			return;
-		}
-
-		if (PendingBundleDownloads.IsEmpty())
+		EInstallBundleRequestFlags InstallFlags = FBaseDownloadGameFeaturePluginState::GetRequestFlags();
+		if (UseAsyncLoading())
 		{
-			StateStatus.SetTransition(EGameFeaturePluginState::Registering);
+			InstallFlags |= EInstallBundleRequestFlags::AsyncMount;
 		}
+		return InstallFlags;
 	}
 
-	virtual void EndState() override
-	{
-		Cleanup();
-	}
-
-	virtual void TryCancelState() override
+	virtual EGameFeaturePluginState GetSuccessTransitionState() const override
 	{
-		// TODO: support pause, cancel etc.  We could possibly be streaming hundreds of assets.
+		return EGameFeaturePluginState::Registering;
 	}
 
-	virtual UE::GameFeatures::FResult TryUpdateProtocolOptions(const FGameFeatureProtocolOptions& NewOptions) override
+	virtual EGameFeaturePluginState GetFailureTransitionState() const override
 	{
-		// TODO: support pause, cancel etc.  We could possibly be streaming hundreds of assets.
-		return FGameFeaturePluginState::TryUpdateProtocolOptions(NewOptions);
+		return EGameFeaturePluginState::ErrorAssetDependencyStreaming;
 	}
 };
 
@@ -2471,7 +2487,7 @@ struct FGameFeaturePluginState_Unregistering : public FGameFeaturePluginState
 	{
 		if (bHasUnloaded)
 		{
-			StateStatus.SetTransition(EGameFeaturePluginState::Unmounting);
+			StateStatus.SetTransition(EGameFeaturePluginState::AssetDependencyStreamOut);
 			return;
 		}
 
@@ -3493,6 +3509,17 @@ bool UGameFeaturePluginStateMachine::AllowAsyncLoading() const
 	return StateProperties.AllowAsyncLoading();
 }
 
+bool UGameFeaturePluginStateMachine::HasAssetStreamingDependencies() const
+{
+	ensure(IsStatusKnown());
+	if (StateProperties.ProtocolMetadata.HasSubtype<FInstallBundlePluginProtocolMetaData>())
+	{
+		const FInstallBundlePluginProtocolMetaData& ProtocolData = StateProperties.ProtocolMetadata.GetSubtype<FInstallBundlePluginProtocolMetaData>();
+		return !ProtocolData.AssetDependencyBundles.IsEmpty();
+	}
+	return false;
+}
+
 void UGameFeaturePluginStateMachine::SetWasLoadedAsBuiltIn()
 {
 	StateProperties.bWasLoadedAsBuiltInGameFeaturePlugin = true;
@@ -3758,7 +3785,7 @@ FString FInstallBundlePluginProtocolMetaData::ToString() const
 	return ReturnedString;
 }
 
-TValueOrError<FInstallBundlePluginProtocolMetaData, void> FInstallBundlePluginProtocolMetaData::FromString(FStringView URLOptionsString)
+TValueOrError<FInstallBundlePluginProtocolMetaData, FString> FInstallBundlePluginProtocolMetaData::FromString(FStringView URLOptionsString)
 {
 	TArray<FName> InstallBundles;
 
@@ -3774,24 +3801,39 @@ TValueOrError<FInstallBundlePluginProtocolMetaData, void> FInstallBundlePluginPr
 	{
 		bParseSuccess = false;
 		UE_LOG(LogGameFeatures, Error, TEXT("Error parsing InstallBundle protocol options URL %.*s"), URLOptionsString.Len(), URLOptionsString.GetData());
-		return MakeError();
+		return MakeError(TEXTVIEW("Bad_PluginURL"));
 	}
 
-	return MakeValue<FInstallBundlePluginProtocolMetaData>(MoveTemp(InstallBundles));
+	// Get any additional bundles for streamed assets
+	TValueOrError<TArray<FName>, FString> MaybeStreamingAssetInstallBundles =
+		UGameFeaturesSubsystem::Get().GetPolicy().GetStreamingAssetInstallBundles(URLOptionsString);
+	if (MaybeStreamingAssetInstallBundles.HasError())
+	{
+		UE_LOG(LogGameFeatures, Error, TEXT("Failed to get streaming asset install bundles from GFP policy. URL %.*s"), URLOptionsString.Len(), URLOptionsString.GetData());
+		return MakeError(MaybeStreamingAssetInstallBundles.StealError());
+	}
+
+	FInstallBundlePluginProtocolMetaData Ret;
+	Ret.InstallBundles = MoveTemp(InstallBundles);
+	Ret.AssetDependencyBundles = MaybeStreamingAssetInstallBundles.StealValue();
+
+	return MakeValue<FInstallBundlePluginProtocolMetaData>(MoveTemp(Ret));
 }
 
-bool FGameFeaturePluginStateMachineProperties::ParseURL()
+TValueOrError<void, FString> FGameFeaturePluginStateMachineProperties::ParseURL()
 {
+	const FStringView BadUrlError = TEXTVIEW("Bad_PluginURL");
+
 	if (!ensureMsgf(!PluginIdentifier.IdentifyingURLSubset.IsEmpty(), TEXT("Unexpected empty IdentifyingURLSubset while parsing URL!")))
 	{
-		return false;
+		return MakeError(BadUrlError);
 	}
 
 	FStringView PluginPathFromURL;
 	FStringView URLOptions;
 	if (!UGameFeaturesSubsystem::ParsePluginURL(PluginIdentifier.GetFullPluginURL(), nullptr, &PluginPathFromURL, &URLOptions))
 	{
-		return false;
+		return MakeError(BadUrlError);
 	}
 
 	PluginInstalledFilename = PluginPathFromURL;
@@ -3800,17 +3842,17 @@ bool FGameFeaturePluginStateMachineProperties::ParseURL()
 	if (PluginInstalledFilename.IsEmpty() || !PluginInstalledFilename.EndsWith(TEXT(".uplugin")))
 	{
 		ensureMsgf(false, TEXT("PluginInstalledFilename must have a uplugin extension. PluginInstalledFilename: %s"), *PluginInstalledFilename);
-		return false;
+		return MakeError(BadUrlError);
 	}
 
 	//Do additional parsing of our Metadata from the options on our remaining URL
 	if (GetPluginProtocol() == EGameFeaturePluginProtocol::InstallBundle)
 	{
-		TValueOrError<FInstallBundlePluginProtocolMetaData, void> MaybeMetaData = FInstallBundlePluginProtocolMetaData::FromString(URLOptions);
+		TValueOrError<FInstallBundlePluginProtocolMetaData, FString> MaybeMetaData = FInstallBundlePluginProtocolMetaData::FromString(URLOptions);
 		if (MaybeMetaData.HasError())
 		{
 			ensureMsgf(false, TEXT("Failure to parse URL %s into a valid FInstallBundlePluginProtocolMetaData"), *PluginIdentifier.GetFullPluginURL());
-			return false;
+			return MakeError(MaybeMetaData.StealError());
 		}
 
 		FInstallBundlePluginProtocolMetaData& MetaData = *ProtocolMetadata.SetSubtype<FInstallBundlePluginProtocolMetaData>();
@@ -3826,7 +3868,7 @@ bool FGameFeaturePluginStateMachineProperties::ParseURL()
 			else
 			{
 				ensureMsgf(false, TEXT("Protocol options type is incorrect for URL %s"), *PluginIdentifier.GetFullPluginURL());
-				return false;
+				return MakeError(BadUrlError);
 			}
 		}
 	}
@@ -3836,13 +3878,13 @@ bool FGameFeaturePluginStateMachineProperties::ParseURL()
 		if (!ProtocolOptions.HasSubtype<FNull>())
 		{
 			ensureMsgf(false, TEXT("Protocol options type is incorrect for URL %s"), *PluginIdentifier.GetFullPluginURL());
-			return false;
+			return MakeError(BadUrlError);
 		}
 	}
 
 	static_assert(static_cast<uint8>(EGameFeaturePluginProtocol::Count) == 3, "Update FGameFeaturePluginStateMachineProperties::ParseURL to handle any new Metadata parsing required for new EGameFeaturePluginProtocol. If no metadata is required just increment this counter.");
 
-	return true;
+	return MakeValue();
 }
 
 UE::GameFeatures::FResult FGameFeaturePluginStateMachineProperties::ValidateProtocolOptionsUpdate(const FGameFeatureProtocolOptions& NewProtocolOptions) const
diff --git a/Engine/Plugins/Runtime/GameFeatures/Source/GameFeatures/Private/GameFeaturePluginStateMachine.h b/Engine/Plugins/Runtime/GameFeatures/Source/GameFeatures/Private/GameFeaturePluginStateMachine.h
--- a/Engine/Plugins/Runtime/GameFeatures/Source/GameFeatures/Private/GameFeaturePluginStateMachine.h
+++ b/Engine/Plugins/Runtime/GameFeatures/Source/GameFeatures/Private/GameFeaturePluginStateMachine.h
@@ -90,35 +90,37 @@ Transition states are expected to transition the machine to another state after
                |             |
                +-^---------+-+
                  |         |
-		   ------~---------~--------------------------------
+           ------~---------~--------------------------------
            |     |         |                               |
         +--v-----+--+    +-v---------+               +-----v--------------+
-        |           |    |           |				 |         !          |
+        |           |    |           |               |         !          |
         |Unmounting |    | Mounting  <---------------> ErrorMounting      |
-        |           |    |           |				 |                    |
-        +--^-----^--+    +--+--------+				 +--------------------+
+        |           |    |           |               |                    |
+        +--^-----^--+    +--+--------+               +--------------------+
            |     |          |
            ------~----------~-------------------------------
                  |          |                              |
-                 |       +--v--------------------+   +-----v-----------------------+
+                 |       +--v--------------------+   +-----+-----------------------+
                  |       |                       |   |         !                   |
                  |       |WaitingForDependencies <---> ErrorWaitingForDependencies |
                  |       |                       |   |                             |
-                 |       +--+--------------------+   +-----------^-----------------+
-                 |          |                                    |
-                 |       +--v----------------------+             |
-                 |       |                         |             |
-                 |       |AssetDependencyStreaming +--------------
-                 |       |                         |
-                 |       +--+----------------------+
-                 |          |
-           ------~----------~-------------------------------
-           |     |          |                              |
-        +--v-----+----+  +--v-------- +              +-----v--------------+
-        |             |  |            |				 |         !          |
+                 |       +-----+-----------------+   +-----------------------------+
+                 |             |
+                 |    ---------~-----------------------------------
+                 |    |        |                                  |
++----------------+----v---+ +--v----------------------+     +-----v-------------------------+
+|                         | |                         |     |             !                 |
+|AssetDependencyStreamOut | |AssetDependencyStreaming <-----> ErrorAssetDependencyStreaming |
+|                         | |                         |     |                               |
++----------------^--------+ +--+----------------------+     +-------------------------------+
+                 |             |
+           ------~-------------~----------------------------
+           |     |             |                           |
+        +--v-----+----+  +-----v----- +              +-----v--------------+
+        |             |  |            |              |         !          |
         |Unregistering|  |Registering <--------------> ErrorRegistering   |
-        |             |  |            |				 |                    |
-        +--------^----+  ++---------- +				 +--------------------+
+        |             |  |            |              |                    |
+        +--------^----+  ++-----------+              +--------------------+
                  |        |
                +-+--------v-+
                |      *     |
@@ -235,10 +237,11 @@ struct FInstallBundlePluginProtocolMetaData
 	FInstallBundlePluginProtocolMetaData(TArray<FName> InInstallBundles) : InstallBundles(MoveTemp(InInstallBundles)) {}
 
 	TArray<FName> InstallBundles;
+	TArray<FName> AssetDependencyBundles;
 
 	/** Functions to convert to/from the URL FString representation of this metadata **/
 	FString ToString() const;
-	static TValueOrError<FInstallBundlePluginProtocolMetaData, void> FromString(FStringView URLOptionsString);
+	static TValueOrError<FInstallBundlePluginProtocolMetaData, FString> FromString(FStringView URLOptionsString);
 };
 
 struct FGameFeatureProtocolMetadata : public TUnion<FInstallBundlePluginProtocolMetaData, FNull>
@@ -325,7 +328,7 @@ struct FGameFeaturePluginStateMachineProperties
 
 	EGameFeaturePluginProtocol GetPluginProtocol() const;
 
-	bool ParseURL();
+	TValueOrError<void, FString> ParseURL();
 
 	/** Checks to see if any invalid data was changed during a URL update. True if data updated was all values expected to be changed. */
 	UE::GameFeatures::FResult ValidateProtocolOptionsUpdate(const FGameFeatureProtocolOptions& NewProtocolOptions) const;
@@ -536,6 +539,9 @@ public:
 	/** Whether this machine is allowed to be asynchronous */
 	bool AllowAsyncLoading() const;
 
+	/** Returns true if the plugin will stream dependencies after installing. Only call if IsStatusKnown is true */
+	bool HasAssetStreamingDependencies() const;
+
 	void SetWasLoadedAsBuiltIn();
 
 	bool WasLoadedAsBuiltIn() const;
diff --git a/Engine/Plugins/Runtime/GameFeatures/Source/GameFeatures/Private/GameFeaturesSubsystem.cpp b/Engine/Plugins/Runtime/GameFeatures/Source/GameFeatures/Private/GameFeaturesSubsystem.cpp
--- a/Engine/Plugins/Runtime/GameFeatures/Source/GameFeatures/Private/GameFeaturesSubsystem.cpp
+++ b/Engine/Plugins/Runtime/GameFeatures/Source/GameFeatures/Private/GameFeaturesSubsystem.cpp
@@ -1494,23 +1494,44 @@ bool UGameFeaturesSubsystem::ShouldUpdatePluginProtocolOptions(const UGameFeatur
 
 bool UGameFeaturesSubsystem::GetGameFeaturePluginInstallPercent(const FString& PluginURL, float& Install_Percent) const
 {
+	// TODO: figure out if doing IAD streaming and include IAD progress if needed
+
 	if (const UGameFeaturePluginStateMachine* StateMachine = FindGameFeaturePluginStateMachine(PluginURL))
 	{
 		if (StateMachine->IsStatusKnown() && StateMachine->IsAvailable())
 		{
 			const FGameFeaturePluginStateInfo& StateInfo = StateMachine->GetCurrentStateInfo();
+
+			float InstallProgress = 0.0f;
 			if (StateInfo.State == EGameFeaturePluginState::Downloading)
 			{
-				Install_Percent = StateInfo.Progress;
+				InstallProgress = StateInfo.Progress;
 			}
 			else if (StateInfo.State >= EGameFeaturePluginState::Installed)
 			{
-				Install_Percent = 1.0f;
+				InstallProgress = 1.0f;
 			}
-			else
+
+			if (!StateMachine->HasAssetStreamingDependencies())
 			{
-				Install_Percent = 0.0f;
+				Install_Percent = InstallProgress;
+				return true;
 			}
+
+			float AssetDependencyProgress = 0.0f;
+			if (StateInfo.State == EGameFeaturePluginState::AssetDependencyStreaming)
+			{
+				AssetDependencyProgress = StateInfo.Progress;
+			}
+			else if(StateInfo.State >= EGameFeaturePluginState::Registering)
+			{
+				AssetDependencyProgress = 1.0f;
+			}
+
+			// Assumuption that most of the progress will be from asset dependencies in this case
+			// For this to be more accurate we'd need to figure out the actual sizes during 
+			// EGameFeaturePluginState::CheckingStatus but this is most likely good enough
+			Install_Percent = 0.2f * InstallProgress + 0.8f * AssetDependencyProgress;
 			return true;
 		}
 	}
@@ -2361,13 +2382,13 @@ struct FGameFeaturePluginPredownloadContext : public FGameFeaturePluginPredownlo
 				continue;
 			}
 
-			TValueOrError<FInstallBundlePluginProtocolMetaData, void> MaybeInstallBundleOptions = FInstallBundlePluginProtocolMetaData::FromString(URL);
+			TValueOrError<FInstallBundlePluginProtocolMetaData, FString> MaybeInstallBundleOptions = FInstallBundlePluginProtocolMetaData::FromString(URL);
 			if (MaybeInstallBundleOptions.HasError())
 			{
 				UE_LOGFMT(LogGameFeatures, Error, "GFP Predownload failed to parse URL {URL}", ("URL", URL));
 				UE::GameFeatures::FResult ErrorResult = MakeError(FString::Printf(TEXT("%.*s%s"),
 					PredownloadErrorNamespace.Len(), PredownloadErrorNamespace.GetData(),
-					TEXT("BadUrl")));
+					*MaybeInstallBundleOptions.GetError()));
 				SetComplete(MoveTemp(ErrorResult));
 				return;
 			}
diff --git a/Engine/Plugins/Runtime/GameFeatures/Source/GameFeatures/Public/GameFeatureTypes.h b/Engine/Plugins/Runtime/GameFeatures/Source/GameFeatures/Public/GameFeatureTypes.h
--- a/Engine/Plugins/Runtime/GameFeatures/Source/GameFeatures/Public/GameFeatureTypes.h
+++ b/Engine/Plugins/Runtime/GameFeatures/Source/GameFeatures/Public/GameFeatureTypes.h
@@ -8,7 +8,11 @@
 
 class FString;
 
-/** GFP States */
+/** 
+  GFP States 
+  Desitination states must be fully ordered, Transistion and Error states should be in between the Destination states they transition to/from.
+  See the state chart in GameFeaturePluginStateMachine.h for reference.
+*/
 #define GAME_FEATURE_PLUGIN_STATE_LIST(XSTATE)	\
 	XSTATE(Uninitialized,					NSLOCTEXT("GameFeatures", "UninitializedStateDisplayName", "Uninitialized"))									/* Unset. Not yet been set up. */ \
 	XSTATE(Terminal,						NSLOCTEXT("GameFeatures", "TerminalStateDisplayName", "Terminal"))												/* Final State before removal of the state machine. */ \
@@ -28,6 +32,8 @@ class FString;
 	XSTATE(ErrorWaitingForDependencies,		NSLOCTEXT("GameFeatures", "ErrorWaitingForDependenciesStateDisplayName", "ErrorWaitingForDependencies"))		/* Error state for Installed -> Registered and Registered -> Installed transitions. */ \
 	XSTATE(ErrorRegistering,				NSLOCTEXT("GameFeatures", "ErrorRegisteringDisplayName", "ErrorRegistering"))									/* Error state for Installed -> Registered and Registered -> Installed transitions. */ \
 	XSTATE(WaitingForDependencies,			NSLOCTEXT("GameFeatures", "WaitingForDependenciesStateDisplayName", "WaitingForDependencies"))					/* Transition state Installed -> Registered. In the process of loading code/content for all dependencies into memory. */ \
+	XSTATE(AssetDependencyStreamOut,		NSLOCTEXT("GameFeatures", "AssetDependencyStreamOutDisplayName", "AssetDependencyStreamOut"))					/* Transition state Registered -> Installed. In the process of streaming out individual assets from dependencies. */ \
+	XSTATE(ErrorAssetDependencyStreaming,	NSLOCTEXT("GameFeatures", "ErrorAssetDependencyStreamingStateDisplayName", "ErrorAssetDependencyStreaming"))	/* Error state for Installed -> Registered and Registered -> Installed transitions. */ \
 	XSTATE(AssetDependencyStreaming,		NSLOCTEXT("GameFeatures", "AssetDependencyStreamingDisplayName", "AssetDependencyStreaming"))					/* Transition state Installed -> Registered. In the process of streaming individual assets from dependencies. */ \
 	XSTATE(Unmounting,						NSLOCTEXT("GameFeatures", "UnmountingStateDisplayName", "Unmounting"))											/* Transition state Registered -> Installed. The content file(s) (i.e. pak file) for the plugin is unmounting. */ \
 	XSTATE(Mounting,						NSLOCTEXT("GameFeatures", "MountingStateDisplayName", "Mounting"))												/* Transition state Installed -> Registered. The content files(s) (i.e. pak file) for the plugin is getting mounted. */ \
diff --git a/Engine/Plugins/Runtime/GameFeatures/Source/GameFeatures/Public/GameFeaturesProjectPolicies.h b/Engine/Plugins/Runtime/GameFeatures/Source/GameFeatures/Public/GameFeaturesProjectPolicies.h
--- a/Engine/Plugins/Runtime/GameFeatures/Source/GameFeatures/Public/GameFeaturesProjectPolicies.h
+++ b/Engine/Plugins/Runtime/GameFeatures/Source/GameFeatures/Public/GameFeaturesProjectPolicies.h
@@ -47,7 +47,7 @@ public:
 	virtual TValueOrError<FString, FString> ResolvePluginDependency(const FString& PluginURL, const FString& DependencyName) const;
 
 	// Called to resolve install bundles for streaming asset dependencies
-	virtual TValueOrError<TArray<FName>, FString> GetStreamingAssetInstallBundles(const FString& PluginURL) const { return MakeValue(); }
+	virtual TValueOrError<TArray<FName>, FString> GetStreamingAssetInstallBundles(FStringView PluginURL) const { return MakeValue(); }
 
 	// Called by code that explicitly wants to load a specific plugin
 	// (e.g., when using a fast cook a game might want to allow explicitly loaded game feature plugins)
diff --git a/Engine/Plugins/Runtime/GameFeatures/Source/GameFeatures/Public/GameFeaturesSubsystem.h b/Engine/Plugins/Runtime/GameFeatures/Source/GameFeatures/Public/GameFeaturesSubsystem.h
--- a/Engine/Plugins/Runtime/GameFeatures/Source/GameFeatures/Public/GameFeaturesSubsystem.h
+++ b/Engine/Plugins/Runtime/GameFeatures/Source/GameFeatures/Public/GameFeaturesSubsystem.h
@@ -672,7 +672,7 @@ private:
 	friend struct FGameFeaturePluginState_Unloading;
 
 	void OnGameFeaturePauseChange(const FGameFeaturePluginIdentifier& PluginIdentifier, const FString& PluginName, FGameFeaturePauseStateChangeContext& Context);
-	friend struct FGameFeaturePluginState_Downloading;
+	friend struct FBaseDownloadGameFeaturePluginState;
 	friend struct FGameFeaturePluginState_Deactivating;
 
 	void OnAssetManagerCreated();
diff --git a/Engine/Plugins/Runtime/GameFeatures/Source/GameFeaturesEditor/Private/SGameFeatureStateWidget.cpp b/Engine/Plugins/Runtime/GameFeatures/Source/GameFeaturesEditor/Private/SGameFeatureStateWidget.cpp
--- a/Engine/Plugins/Runtime/GameFeatures/Source/GameFeaturesEditor/Private/SGameFeatureStateWidget.cpp
+++ b/Engine/Plugins/Runtime/GameFeatures/Source/GameFeaturesEditor/Private/SGameFeatureStateWidget.cpp
@@ -65,7 +65,7 @@ FText SGameFeatureStateWidget::GetDisplayNameOfState(EGameFeaturePluginState Sta
 
 FText SGameFeatureStateWidget::GetTooltipOfState(EGameFeaturePluginState State)
 {
-	static_assert((int32)EGameFeaturePluginState::MAX == 35, "");
+	static_assert((int32)EGameFeaturePluginState::MAX == 37, "");
 
 	switch (State)
 	{
@@ -105,6 +105,10 @@ FText SGameFeatureStateWidget::GetTooltipOfState(EGameFeaturePluginState State)
 		return LOCTEXT("StateTooltip_ErrorRegistering", "Error state for Installed -> Registered and Registered -> Installed transitions.");
 	case EGameFeaturePluginState::WaitingForDependencies:
 		return LOCTEXT("StateTooltip_WaitingForDependencies", "Transition state Installed -> Registered. In the process of loading code/content for all dependencies into memory.");
+	case EGameFeaturePluginState::AssetDependencyStreamOut:
+		return LOCTEXT("StateTooltip_AssetDependencyStreamout", "Transition state Registered -> Installed. In the process of streaming out individual assets from dependencies.");
+	case EGameFeaturePluginState::ErrorAssetDependencyStreaming:
+		return LOCTEXT("StateTooltip_ErrorAssetDependencyStreaming", "Error state for Installed -> Registered and Registered -> Installed transitions.");
 	case EGameFeaturePluginState::AssetDependencyStreaming:
 		return LOCTEXT("StateTooltip_AssetDependencyStreaming", "Transition state Installed -> Registered. In the process of streaming individual assets from dependencies.");
 	case EGameFeaturePluginState::Unmounting:
