Add a new GameFeatureAction to support actor factories inside of a GameFeaturePlugin.
Adding new config variable for ActorFactories bShouldAutoRegister so that actor factories do not auto register optionally for GameFeaturePlugins.

diff --git a/Engine/Plugins/Runtime/GameFeatures/Source/GameFeatures/GameFeatures.Build.cs b/Engine/Plugins/Runtime/GameFeatures/Source/GameFeatures/GameFeatures.Build.cs
--- a/Engine/Plugins/Runtime/GameFeatures/Source/GameFeatures/GameFeatures.Build.cs
+++ b/Engine/Plugins/Runtime/GameFeatures/Source/GameFeatures/GameFeatures.Build.cs
@@ -39,6 +39,7 @@ namespace UnrealBuildTool.Rules
 					new string[]
 					{
 						"UnrealEd",
+						"PlacementMode",
 						"PluginUtils",
 					}
 				);
diff --git a/Engine/Plugins/Runtime/GameFeatures/Source/GameFeatures/Private/GameFeatureAction_AddActorFactory.cpp b/Engine/Plugins/Runtime/GameFeatures/Source/GameFeatures/Private/GameFeatureAction_AddActorFactory.cpp
new file mode 100644
--- /dev/null
+++ b/Engine/Plugins/Runtime/GameFeatures/Source/GameFeatures/Private/GameFeatureAction_AddActorFactory.cpp
@@ -0,0 +1,106 @@
+// Copyright Epic Games, Inc. All Rights Reserved.
+
+#include "GameFeatureAction_AddActorFactory.h"
+#include "GameFeatureData.h"
+#include "Misc/MessageDialog.h"
+
+#if WITH_EDITOR
+#include "IPlacementModeModule.h"
+#endif // WITH_EDITOR
+
+#include UE_INLINE_GENERATED_CPP_BY_NAME(GameFeatureAction_AddActorFactory)
+
+#define LOCTEXT_NAMESPACE "GameFeatures"
+
+//////////////////////////////////////////////////////////////////////
+// UGameFeatureAction_AddActorFactory
+
+DEFINE_LOG_CATEGORY_STATIC(LogAddActorFactory, Log, All);
+
+void UGameFeatureAction_AddActorFactory::OnGameFeatureRegistering()
+{
+	AddActorFactory();
+}
+
+void UGameFeatureAction_AddActorFactory::OnGameFeatureUnregistering()
+{
+	RemoveActorFactory();
+}
+
+#if WITH_EDITOR
+void UGameFeatureAction_AddActorFactory::PostRename(UObject* OldOuter, const FName OldName)
+{
+	Super::PostRename(OldOuter, OldName);
+
+	// If OldOuter is not GetTransientPackage(), but GetOuter() is GetTransientPackage(), then you were trashed.
+	const UObject* MyOuter = GetOuter();
+	const UPackage* TransientPackage = GetTransientPackage();
+	if (OldOuter != TransientPackage && MyOuter == TransientPackage)
+	{
+		RemoveActorFactory();
+	}
+}
+
+void UGameFeatureAction_AddActorFactory::PostEditChangeProperty(FPropertyChangedEvent& PropertyChangedEvent)
+{
+	Super::PostEditChangeProperty(PropertyChangedEvent);
+
+	const FName PropertyName = PropertyChangedEvent.GetPropertyName();
+	if (PropertyName == GET_MEMBER_NAME_CHECKED(UGameFeatureAction_AddActorFactory, ActorFactory))
+	{
+		RemoveActorFactory();
+		AddActorFactory();
+	}
+}
+#endif // WITH_EDITOR
+
+void UGameFeatureAction_AddActorFactory::AddActorFactory()
+{
+#if WITH_EDITOR
+	if (ActorFactory.IsNull())
+	{
+		UE_LOG(LogAddActorFactory, Warning, TEXT("ActorFactory is null. Unable to add factory"));
+		return;
+	}
+	if (UClass* FactoryClass = ActorFactory.LoadSynchronous())
+	{
+		if (!FactoryClass->IsChildOf(UActorFactory::StaticClass()))
+		{
+			UE_LOG(LogAddActorFactory, Error, TEXT("ActorFactory (%s) was not an ActorFactory class"), *FactoryClass->GetName());
+			FMessageDialog::Open(EAppMsgType::Ok, LOCTEXT("AddActorFactory_BadSubclass", "Selected class was not an ActorFactory class."));
+			ActorFactory.Reset();
+			return;
+		}
+		UE_LOG(LogAddActorFactory, Verbose, TEXT("Adding actor factory %s"), *FactoryClass->GetName());
+
+		UActorFactory* NewFactory = NewObject<UActorFactory>(GetTransientPackage(), FactoryClass);
+		if (NewFactory->bShouldAutoRegister)
+		{
+			FMessageDialog::Open(EAppMsgType::Ok, LOCTEXT("AddActorFactory_AutoRegister", "The selected actor factory is set to auto register. Set the config variable bShouldAutoRegister to false before using this action."));
+			ActorFactory.Reset();
+			return;
+		}
+		GEditor->ActorFactories.Add(NewFactory);
+		AddedFactory = NewFactory;
+		IPlacementModeModule::Get().RegenerateItemsForCategory(FBuiltInPlacementCategories::AllClasses());
+	}
+#endif // WITH_EDITOR
+}
+
+void UGameFeatureAction_AddActorFactory::RemoveActorFactory()
+{
+#if WITH_EDITOR
+	if (UActorFactory* FactoryToRemove = Cast<UActorFactory>(AddedFactory.Get()))
+	{
+		UE_LOG(LogAddActorFactory, Verbose, TEXT("Removing actor factory %s"), *FactoryToRemove->GetName());
+
+		GEditor->ActorFactories.Remove(FactoryToRemove);
+		AddedFactory.Reset();
+		IPlacementModeModule::Get().RegenerateItemsForCategory(FBuiltInPlacementCategories::AllClasses());
+	}
+#endif // WITH_EDITOR
+}
+
+//////////////////////////////////////////////////////////////////////
+
+#undef LOCTEXT_NAMESPACE
diff --git a/Engine/Plugins/Runtime/GameFeatures/Source/GameFeatures/Public/GameFeatureAction_AddActorFactory.h b/Engine/Plugins/Runtime/GameFeatures/Source/GameFeatures/Public/GameFeatureAction_AddActorFactory.h
new file mode 100644
--- /dev/null
+++ b/Engine/Plugins/Runtime/GameFeatures/Source/GameFeatures/Public/GameFeatureAction_AddActorFactory.h
@@ -0,0 +1,55 @@
+// Copyright Epic Games, Inc. All Rights Reserved.
+
+#pragma once
+
+#include "GameFeatureAction.h"
+#include "GameFeatureAction_AddActorFactory.generated.h"
+
+class UActorComponent;
+class UActorFactory;
+
+//////////////////////////////////////////////////////////////////////
+// UGameFeatureAction_AddActorFactory
+
+/**
+ * GameFeatureAction to add an actor factory when this plugin registers.
+ * Useful for factories that might load BP classes within a plugin which might not have been discovered yet.
+ */
+UCLASS(MinimalAPI, meta=(DisplayName="Add Actor Factory"), Config=Engine)
+class UGameFeatureAction_AddActorFactory final : public UGameFeatureAction
+{
+	GENERATED_BODY()
+
+public:
+	//~Start of UGameFeatureAction interface
+	virtual void OnGameFeatureRegistering() override;
+	virtual void OnGameFeatureUnregistering() override;
+	//~End of UGameFeatureAction interface
+
+#if WITH_EDITOR
+
+	/** UObject overrides */
+	virtual void PostRename(UObject* OldOuter, const FName OldName) override;
+	virtual void PostEditChangeProperty(FPropertyChangedEvent& PropertyChangedEvent) override;
+	/** ~UObject overrides */
+#endif // WITH_EDITOR
+	
+protected:
+#if WITH_EDITORONLY_DATA
+	/**
+	 * The actor factory class to add once this plugin registers
+	 * Actor factories should be setup with bShouldAutoRegister so that they do not register during engine boot.
+	 */
+	UPROPERTY(EditAnywhere, Category = "Actor Factory")
+	TSoftClassPtr<UObject> ActorFactory;
+
+	TWeakObjectPtr<UObject> AddedFactory;
+#endif
+
+	void AddActorFactory();
+	void RemoveActorFactory();
+};
+
+#if UE_ENABLE_INCLUDE_ORDER_DEPRECATED_IN_5_2
+#include "CoreMinimal.h"
+#endif
diff --git a/Engine/Source/Editor/UnrealEd/Classes/ActorFactories/ActorFactory.h b/Engine/Source/Editor/UnrealEd/Classes/ActorFactories/ActorFactory.h
--- a/Engine/Source/Editor/UnrealEd/Classes/ActorFactories/ActorFactory.h
+++ b/Engine/Source/Editor/UnrealEd/Classes/ActorFactories/ActorFactory.h
@@ -39,6 +39,10 @@ class UActorFactory : public UObject, public IAssetFactoryInterface
 	UPROPERTY(config)
 	FString NewActorClassName;
 
+	/** should this actor factory auto register with the editor factory system. For factories within a GameFeaturePlugin this should be false. */
+	UPROPERTY(config)
+	bool bShouldAutoRegister = true;
+
 	/**  AActor  subclass this ActorFactory creates. */
 	UPROPERTY()
 	TSubclassOf<AActor>  NewActorClass;
diff --git a/Engine/Source/Editor/UnrealEd/Private/EditorEngine.cpp b/Engine/Source/Editor/UnrealEd/Private/EditorEngine.cpp
--- a/Engine/Source/Editor/UnrealEd/Private/EditorEngine.cpp
+++ b/Engine/Source/Editor/UnrealEd/Private/EditorEngine.cpp
@@ -952,7 +952,7 @@ void UEditorEngine::InitEditor(IEngineLoop* InEngineLoop)
 					{
 						VolumeFactoryClasses.Add(TestClass);
 					}
-					else
+					else if (TestClass->GetDefaultObject<UActorFactory>()->bShouldAutoRegister)
 					{
 						UActorFactory* NewFactory = NewObject<UActorFactory>(GetTransientPackage(), TestClass);
 						check(NewFactory);
