Fix for the default / basic set of place actors items being empty because they were created at a time where their actor factory isn't loaded yet. Don't force load the placement mode module to refresh its entries if it hasn't loaded yet.

diff --git a/Engine/Plugins/Runtime/GameFeatures/Source/GameFeatures/Private/GameFeatureAction_AddActorFactory.cpp b/Engine/Plugins/Runtime/GameFeatures/Source/GameFeatures/Private/GameFeatureAction_AddActorFactory.cpp
--- a/Engine/Plugins/Runtime/GameFeatures/Source/GameFeatures/Private/GameFeatureAction_AddActorFactory.cpp
+++ b/Engine/Plugins/Runtime/GameFeatures/Source/GameFeatures/Private/GameFeatureAction_AddActorFactory.cpp
@@ -82,7 +82,11 @@ void UGameFeatureAction_AddActorFactory::AddActorFactory()
 		}
 		GEditor->ActorFactories.Add(NewFactory);
 		AddedFactory = NewFactory;
-		IPlacementModeModule::Get().RegenerateItemsForCategory(FBuiltInPlacementCategories::AllClasses());
+
+		if (IPlacementModeModule::IsAvailable())
+		{
+			IPlacementModeModule::Get().RegenerateItemsForCategory(FBuiltInPlacementCategories::AllClasses());
+		}
 	}
 #endif // WITH_EDITOR
 }
@@ -96,7 +100,11 @@ void UGameFeatureAction_AddActorFactory::RemoveActorFactory()
 
 		GEditor->ActorFactories.Remove(FactoryToRemove);
 		AddedFactory.Reset();
-		IPlacementModeModule::Get().RegenerateItemsForCategory(FBuiltInPlacementCategories::AllClasses());
+
+		if (IPlacementModeModule::IsAvailable())
+		{
+			IPlacementModeModule::Get().RegenerateItemsForCategory(FBuiltInPlacementCategories::AllClasses());
+		}
 	}
 #endif // WITH_EDITOR
 }
