Add new GameFeaturePolicy function GetGameFeaturePluginURL to get the PluginURL of a valid IPlugin.
This can be used for gamecode to choose if a plugin should be file or install protocol.

diff --git a/Engine/Plugins/Runtime/GameFeatures/Source/GameFeatures/Private/GameFeaturesProjectPolicies.cpp b/Engine/Plugins/Runtime/GameFeatures/Source/GameFeatures/Private/GameFeaturesProjectPolicies.cpp
--- a/Engine/Plugins/Runtime/GameFeatures/Source/GameFeatures/Private/GameFeaturesProjectPolicies.cpp
+++ b/Engine/Plugins/Runtime/GameFeatures/Source/GameFeatures/Private/GameFeaturesProjectPolicies.cpp
@@ -50,6 +50,26 @@ const TArray<FName> UDefaultGameFeaturesProjectPolicies::GetPreloadBundleStateFo
 	return FeatureBundles;
 }
 
+bool UGameFeaturesProjectPolicies::GetGameFeaturePluginURL(const TSharedRef<IPlugin>& Plugin, FString& OutPluginURL) const
+{
+	// It could still be a GFP, but state machine may not have been created for it yet
+	// Check if it is a built-in GFP
+	const FString& PluginDescriptorFilename = Plugin->GetDescriptorFileName();
+	if (!PluginDescriptorFilename.IsEmpty())
+	{
+		if (GetDefault<UGameFeaturesSubsystemSettings>()->IsValidGameFeaturePlugin(FPaths::ConvertRelativePathToFull(PluginDescriptorFilename)))
+		{
+			OutPluginURL = UGameFeaturesSubsystem::GetPluginURL_FileProtocol(PluginDescriptorFilename);
+		}
+		else
+		{
+			OutPluginURL = TEXT("");
+		}
+		return true;
+	}
+	return false;
+}
+
 bool UGameFeaturesProjectPolicies::WillPluginBeCooked(const FString& PluginFilename, const FGameFeaturePluginDetails& PluginDetails) const
 {
 	return true;
@@ -68,15 +88,7 @@ TValueOrError<FString, FString> UGameFeaturesProjectPolicies::ResolvePluginDepen
 	// Check if the dependency plugin exists yet (should be true for all built-in plugins)
 	else if (TSharedPtr<IPlugin> DependencyPlugin = IPluginManager::Get().FindPlugin(DependencyName))
 	{
-		// It could still be a GFP, but state machine may not have been created for it yet
-		// Check if it is a built-in GFP
-		if (!DependencyPlugin->GetDescriptorFileName().IsEmpty() &&
-			GetDefault<UGameFeaturesSubsystemSettings>()->IsValidGameFeaturePlugin(FPaths::ConvertRelativePathToFull(DependencyPlugin->GetDescriptorFileName())))
-		{
-			DependencyURL = UGameFeaturesSubsystem::GetPluginURL_FileProtocol(DependencyPlugin->GetDescriptorFileName());
-		}
-
-		bResolvedDependency = true;
+		bResolvedDependency = GetGameFeaturePluginURL(DependencyPlugin.ToSharedRef(), DependencyURL);
 	}
 
 	if (bResolvedDependency)
diff --git a/Engine/Plugins/Runtime/GameFeatures/Source/GameFeatures/Private/GameFeaturesSubsystem.cpp b/Engine/Plugins/Runtime/GameFeatures/Source/GameFeatures/Private/GameFeaturesSubsystem.cpp
--- a/Engine/Plugins/Runtime/GameFeatures/Source/GameFeatures/Private/GameFeaturesSubsystem.cpp
+++ b/Engine/Plugins/Runtime/GameFeatures/Source/GameFeatures/Private/GameFeaturesSubsystem.cpp
@@ -2105,24 +2105,16 @@ bool UGameFeaturesSubsystem::GetBuiltInGameFeaturePluginDetails(const TSharedRef
 	// Ideally this would work with any protocol, but for current uses cases the exact protocol doesn't seem to matter.
 
 	const FString& PluginDescriptorFilename = Plugin->GetDescriptorFileName();
-
 	// Make sure you are in a game feature plugins folder. All GameFeaturePlugins are rooted in a GameFeatures folder.
 	if (!PluginDescriptorFilename.IsEmpty() && GetDefault<UGameFeaturesSubsystemSettings>()->IsValidGameFeaturePlugin(FPaths::ConvertRelativePathToFull(PluginDescriptorFilename)) && FPaths::FileExists(PluginDescriptorFilename))
 	{
-		bool bIsFileProtocol = true;
-		if (GetPluginURLByName(Plugin->GetName(), OutPluginURL))
-		{
-			bIsFileProtocol = UGameFeaturesSubsystem::IsPluginURLProtocol(OutPluginURL, EGameFeaturePluginProtocol::File);
-		}
-		else
-		{
-			OutPluginURL = GetPluginURL_FileProtocol(PluginDescriptorFilename);
-		}
-
-		if (bIsFileProtocol)
+		const FString PluginName = Plugin->GetName();
+		const bool bFoundPluginURL = GetPluginURLByName(PluginName, OutPluginURL);
+		if (!bFoundPluginURL)
 		{
-			return GetGameFeaturePluginDetailsInternal(PluginDescriptorFilename, OutPluginDetails);
+			GameSpecificPolicies->GetGameFeaturePluginURL(Plugin, OutPluginURL);
 		}
+		return GetGameFeaturePluginDetailsInternal(PluginDescriptorFilename, OutPluginDetails);
 	}
 
 	return false;
diff --git a/Engine/Plugins/Runtime/GameFeatures/Source/GameFeatures/Public/GameFeaturesProjectPolicies.h b/Engine/Plugins/Runtime/GameFeatures/Source/GameFeatures/Public/GameFeaturesProjectPolicies.h
--- a/Engine/Plugins/Runtime/GameFeatures/Source/GameFeatures/Public/GameFeaturesProjectPolicies.h
+++ b/Engine/Plugins/Runtime/GameFeatures/Source/GameFeatures/Public/GameFeaturesProjectPolicies.h
@@ -37,6 +37,9 @@ public:
 	// Actions can use this to decide what to load at runtime
 	virtual void GetGameFeatureLoadingMode(bool& bLoadClientData, bool& bLoadServerData) const { bLoadClientData = true; bLoadServerData = true; }
 
+	// Called to determine the plugin URL for a given known Plugin. Can be used if the policy wants to deliver non file based URLs.
+	virtual bool GetGameFeaturePluginURL(const TSharedRef<IPlugin>& Plugin, FString& OutPluginURL) const;
+
 	// Called to determine if a plugin is allowed to be loaded or not
 	// (e.g., when doing a fast cook a game might want to disable some or all game feature plugins)
 	virtual bool IsPluginAllowed(const FString& PluginURL) const { return true; }
