Allow legacy GFPs to register with the config cache even when they use the old ini naming scheme

#rnx

diff --git a/Engine/Plugins/Runtime/GameFeatures/Source/GameFeatures/Private/GameFeatureData.cpp b/Engine/Plugins/Runtime/GameFeatures/Source/GameFeatures/Private/GameFeatureData.cpp
--- a/Engine/Plugins/Runtime/GameFeatures/Source/GameFeatures/Private/GameFeatureData.cpp
+++ b/Engine/Plugins/Runtime/GameFeatures/Source/GameFeatures/Private/GameFeatureData.cpp
@@ -86,6 +86,9 @@ void UGameFeatureData::InitializeBasePluginIniFile(const FString& PluginInstalle
 {
 	const FString PluginName = FPaths::GetBaseFilename(PluginInstalledFilename);
 
+	static bool bUseNewDynamicLayers = IConsoleManager::Get().FindConsoleVariable(TEXT("ini.UseNewDynamicLayers"))->GetInt() != 0;
+	bool bIncludePluginNameInBranchName = true;
+
 	// DEPRECATED NAMING PATH - must keep because these files are read in as a single file, not in a hierarchical way, so 
 	// they don't have the + syntax for arrays
 	{
@@ -106,6 +109,13 @@ void UGameFeatureData::InitializeBasePluginIniFile(const FString& PluginInstalle
 			// This is the deprecated loading path that doesn't handle cases like + in arrays
 			UE_LOG(LogGameFeatures, Log, TEXT("[GameFeatureData %s]: Loaded deprecated config %s, rename to start with Default for normal parsing"), *GetPathNameSafe(this), *PluginConfigFilename);
 
+			// register this plugin, so the ConfigContext.Load, and future loads, know about it
+			if (bUseNewDynamicLayers)
+			{
+				TSharedPtr<IPlugin> Plugin = IPluginManager::Get().FindPlugin(PluginName);
+				FConfigCacheIni::RegisterPlugin(*Plugin->GetName(), Plugin->GetBaseDir(), Plugin->GetExtensionBaseDirs(), DynamicLayerPriority::GameFeature, bIncludePluginNameInBranchName);
+			}
+
 			FCoreRedirects::ReadRedirectsFromIni(PluginConfigFilename);
 			ReloadConfigs(PluginConfig);
 
@@ -113,7 +123,6 @@ void UGameFeatureData::InitializeBasePluginIniFile(const FString& PluginInstalle
 		}
 	}
 
-	static bool bUseNewDynamicLayers = IConsoleManager::Get().FindConsoleVariable(TEXT("ini.UseNewDynamicLayers"))->GetInt() != 0;
 	if (bUseNewDynamicLayers)
 	{
 		TSharedPtr<IPlugin> Plugin = IPluginManager::Get().FindPlugin(PluginName);
@@ -123,7 +132,6 @@ void UGameFeatureData::InitializeBasePluginIniFile(const FString& PluginInstalle
 		UE_LOG(LogGameFeatures, Verbose, TEXT("Loading GameFeature base plugin hierarchy for %s"), *PluginName);
 
 		// register this plugin, so the ConfigContext.Load, and future loads, know about it
-		bool bIncludePluginNameInBranchName = true;
 		FConfigCacheIni::RegisterPlugin(*Plugin->GetName(), Plugin->GetBaseDir(), Plugin->GetExtensionBaseDirs(), DynamicLayerPriority::GameFeature, bIncludePluginNameInBranchName);
 
 		// load the plugin inis and track the modified sections
