Fix missing config reload for GFPs with a single base ini file while using dynamic config layers

#rnx

diff --git a/Engine/Plugins/Runtime/GameFeatures/Source/GameFeatures/Private/GameFeatureData.cpp b/Engine/Plugins/Runtime/GameFeatures/Source/GameFeatures/Private/GameFeatureData.cpp
--- a/Engine/Plugins/Runtime/GameFeatures/Source/GameFeatures/Private/GameFeatureData.cpp
+++ b/Engine/Plugins/Runtime/GameFeatures/Source/GameFeatures/Private/GameFeatureData.cpp
@@ -140,6 +140,12 @@ void UGameFeatureData::InitializeBasePluginIniFile(const FString& PluginInstalle
 		Context.ConfigFileTag = *Plugin->GetName();
 		Context.Load(*PluginName, PluginConfigFilename);
 
+		FConfigFile* PluginConfig = GConfig->Find(PluginConfigFilename);
+		if (PluginConfig)
+		{
+			ReloadConfigs(*PluginConfig);
+		}
+
 		// @todo move this into ReloadObjectsFromModifiedConfigSections?
 		FCoreRedirects::ReadRedirectsFromIni(PluginConfigFilename);
 
