Default built in pluing should load the ini even if an install bundle type.

diff --git a/Engine/Plugins/Runtime/GameFeatures/Source/GameFeatures/Private/GameFeaturesSubsystem.cpp b/Engine/Plugins/Runtime/GameFeatures/Source/GameFeatures/Private/GameFeaturesSubsystem.cpp
--- a/Engine/Plugins/Runtime/GameFeatures/Source/GameFeatures/Private/GameFeaturesSubsystem.cpp
+++ b/Engine/Plugins/Runtime/GameFeatures/Source/GameFeatures/Private/GameFeaturesSubsystem.cpp
@@ -1790,6 +1790,12 @@ void UGameFeaturesSubsystem::LoadBuiltInGameFeaturePlugin(const TSharedRef<IPlug
 			if (bShouldProcess)
 			{
 				FGameFeatureProtocolOptions ProtocolOptions;
+				if (UGameFeaturesSubsystem::GetPluginURLProtocol(PluginURL) == EGameFeaturePluginProtocol::InstallBundle)
+				{
+					FInstallBundlePluginProtocolOptions InstallBundleOptions;
+					InstallBundleOptions.bAllowIniLoading = true;
+					ProtocolOptions = FGameFeatureProtocolOptions(InstallBundleOptions);
+				}
 				ProtocolOptions.bForceSyncLoading = BehaviorOptions.bForceSyncLoading;
 				ProtocolOptions.bLogWarningOnForcedDependencyCreation = BehaviorOptions.bLogWarningOnForcedDependencyCreation;
 				ProtocolOptions.bLogErrorOnForcedDependencyCreation = BehaviorOptions.bLogErrorOnForcedDependencyCreation;
