Remove restriction that the package localization cache can only be updated on the game thread

This restriction was added before the asset registry itself was made thread-safe, and is no longer needed

This allows us to remove all the places that were pessimistically calling ConditionalUpdateCache from the GT to avoid warnings from the ALT, as the ALT can now update the cache itself as needed. This avoids the per-plugin/GFP hit when calling ConditionalUpdateCache after mounting/loading the asset registry state, and in the editor means that nothing ever needs to call ConditionalUpdateCache (as the editor doesn't partake in localized asset loading).

diff --git a/Engine/Plugins/Runtime/GameFeatures/Source/GameFeatures/Private/GameFeaturePluginStateMachine.cpp b/Engine/Plugins/Runtime/GameFeatures/Source/GameFeatures/Private/GameFeaturePluginStateMachine.cpp
--- a/Engine/Plugins/Runtime/GameFeatures/Source/GameFeatures/Private/GameFeaturePluginStateMachine.cpp
+++ b/Engine/Plugins/Runtime/GameFeatures/Source/GameFeatures/Private/GameFeaturePluginStateMachine.cpp
@@ -1883,7 +1883,6 @@ struct FGameFeaturePluginState_Mounting : public FGameFeaturePluginState
 	UE::GameFeatures::FResult Result;
 	ESubState StartedSubStates = ESubState::None;
 	ESubState CompletedSubStates = ESubState::None;
-	bool bLoadedAssetRegistryState = false;
 	bool bCheckedRealtimeMode = false;
 	bool bForceMonolithicShaderLibrary = true;	// use monolithic unless a DLC plugin is chunked
 
@@ -1966,7 +1965,6 @@ struct FGameFeaturePluginState_Mounting : public FGameFeaturePluginState
 		Result = MakeValue();
 		StartedSubStates = ESubState::None;
 		CompletedSubStates = ESubState::None;
-		bLoadedAssetRegistryState = false;
 		bCheckedRealtimeMode = false;
 		bForceMonolithicShaderLibrary = false;
 
@@ -2189,6 +2187,16 @@ struct FGameFeaturePluginState_Mounting : public FGameFeaturePluginState
 			}
 		}
 
+		auto RefreshPackageLocalizationCacheForPlugin = [NewlyMountedPlugin]()
+		{
+			// We need to refresh the package localization cache for a GFP if it loaded cooked asset registry state, 
+			// as we need the asset registry data to correctly build the package localization cache for the GFP
+			if (NewlyMountedPlugin && NewlyMountedPlugin->CanContainContent())
+			{
+				FPackageLocalizationManager::Get().InvalidateRootSourcePath(NewlyMountedPlugin->GetMountedAssetPath());
+			}
+		};
+
 		if (!UseAsyncLoading())
 		{
 			FAssetRegistryState PluginAssetRegistryState;
@@ -2196,19 +2204,19 @@ struct FGameFeaturePluginState_Mounting : public FGameFeaturePluginState
 			{
 				IAssetRegistry& AssetRegistry = UAssetManager::Get().GetAssetRegistry();
 				AssetRegistry.AppendState(PluginAssetRegistryState);
+				RefreshPackageLocalizationCacheForPlugin();
 			}
 			else
 			{
 				Result = GetErrorResult(TEXT("Failed_To_Load_Plugin_AssetRegistry"));
 			}
 
-			bLoadedAssetRegistryState = true;
 			CompletedSubStates |= ESubState::LoadAssetRegistry;
 			return;
 		}
 
 		const bool bForceSyncAssetRegistryAppend = UE::GameFeatures::CVarForceSyncAssetRegistryAppend.GetValueOnGameThread();
-		UE::Tasks::Launch(UE_SOURCE_LOCATION, [this, PluginAssetRegistry=MoveTemp(PluginAssetRegistry), bForceSyncAssetRegistryAppend]
+		UE::Tasks::Launch(UE_SOURCE_LOCATION, [this, PluginAssetRegistry=MoveTemp(PluginAssetRegistry), bForceSyncAssetRegistryAppend, RefreshPackageLocalizationCacheForPlugin]
 		{
 			bool bSuccess = false;
 			TSharedPtr<FAssetRegistryState> PluginAssetRegistryState = MakeShared<FAssetRegistryState>();
@@ -2218,11 +2226,12 @@ struct FGameFeaturePluginState_Mounting : public FGameFeaturePluginState
 				if (!bForceSyncAssetRegistryAppend)
 				{
 					AssetRegistry.AppendState(*PluginAssetRegistryState);
+					RefreshPackageLocalizationCacheForPlugin();
 				}
 				bSuccess = true;
 			}
 
-			ExecuteOnGameThread(UE_SOURCE_LOCATION, [this, PluginAssetRegistryState, bSuccess, bForceSyncAssetRegistryAppend]
+			ExecuteOnGameThread(UE_SOURCE_LOCATION, [this, PluginAssetRegistryState, bSuccess, bForceSyncAssetRegistryAppend, RefreshPackageLocalizationCacheForPlugin]
 			{
 				TRACE_CPUPROFILER_EVENT_SCOPE(GFP_Mounting_ARComplete);
 
@@ -2234,9 +2243,9 @@ struct FGameFeaturePluginState_Mounting : public FGameFeaturePluginState
 				{
 					IAssetRegistry& AssetRegistry = UAssetManager::Get().GetAssetRegistry();
 					AssetRegistry.AppendState(*PluginAssetRegistryState);
+					RefreshPackageLocalizationCacheForPlugin();
 				}
 
-				bLoadedAssetRegistryState = true;
 				CompletedSubStates |= ESubState::LoadAssetRegistry;
 				UpdateStateMachineImmediate();
 			});
@@ -2284,18 +2293,6 @@ struct FGameFeaturePluginState_Mounting : public FGameFeaturePluginState
 		// Post-mount
 		if (bComplete)
 		{
-			if (bLoadedAssetRegistryState)
-			{
-				// We need to refresh the package localization cache for a GFP if it loaded cooked asset registry state, 
-				// as we need the asset registry data to correctly build the package localization cache for the GFP
-				if (TSharedPtr<IPlugin> Plugin = IPluginManager::Get().FindPlugin(StateProperties.PluginName);
-					Plugin && Plugin->CanContainContent())
-				{
-					FPackageLocalizationManager::Get().InvalidateRootSourcePath(Plugin->GetMountedAssetPath());
-					FPackageLocalizationManager::Get().ConditionalUpdateCache();
-				}
-			}
-
 			FGameFeaturePostMountingContext Context(StateProperties.PluginName, [this](FStringView InPauserTag) { OnPostMountPauserCompleted(InPauserTag); });
 			NumExpectedPostMountPausers = INDEX_NONE;
 			UGameFeaturesSubsystem::Get().OnGameFeaturePostMounting(StateProperties.PluginName, StateProperties.PluginIdentifier, Context);
diff --git a/Engine/Source/Runtime/CoreUObject/Private/Internationalization/PackageLocalizationCache.cpp b/Engine/Source/Runtime/CoreUObject/Private/Internationalization/PackageLocalizationCache.cpp
--- a/Engine/Source/Runtime/CoreUObject/Private/Internationalization/PackageLocalizationCache.cpp
+++ b/Engine/Source/Runtime/CoreUObject/Private/Internationalization/PackageLocalizationCache.cpp
@@ -28,12 +28,6 @@ void FPackageLocalizationCultureCache::ConditionalUpdateCache_NoLock()
 		return;
 	}
 
-	if (!IsInGameThread())
-	{
-		UE_LOG(LogPackageLocalizationCache, Warning, TEXT("Skipping the cache update for %d pending package path(s) due to a cache request from a non-game thread. Some localized packages may be missed for this query."), PendingSourceRootPathsToSearch.Num());
-		return;
-	}
-
 	SCOPED_BOOT_TIMING("FPackageLocalizationCultureCache::ConditionalUpdateCache_NoLock");
 	const double CacheStartTime = FPlatformTime::Seconds();
 
@@ -333,12 +327,6 @@ void FPackageLocalizationCache::ConditionalUpdatePackageNameToAssetGroupCache_No
 		return;
 	}
 
-	if (!IsInGameThread())
-	{
-		UE_LOG(LogPackageLocalizationCache, Warning, TEXT("Skipping the cache update for the package asset groups due to a cache request from a non-game thread. Some localized packages may be missed for this query."));
-		return;
-	}
-
 	PackageNameToAssetGroup.Reset();
 	for (const auto& AssetClassGroupPair : AssetClassesToAssetGroups)
 	{
@@ -383,16 +371,18 @@ void FPackageLocalizationCache::HandleCultureChanged()
 	const FString CurrentCultureName = FInternationalization::Get().GetCurrentLanguage()->GetName();
 	CurrentCultureCache = FindOrAddCacheForCulture_NoLock(CurrentCultureName);
 
-	// We expect culture changes to happen on the game thread, so update the cache now while it is likely safe to do so
-	// (ConditionalUpdateCache will internally check that this is currently the game thread before allowing the update)
-	const TArray<FCultureRef> CurrentCultures = FInternationalization::Get().GetCurrentCultures(/*bIncludeLanguage*/true, /*bIncludeLocale*/false, /*bIncludeAssetGroups*/true);
-	for (const FCultureRef& CurrentCulture : CurrentCultures)
+	if (!GIsEditor)
 	{
-		if (TSharedPtr<FPackageLocalizationCultureCache> CultureCache = FindOrAddCacheForCulture_NoLock(CurrentCulture->GetName()))
+		// Preemptively update the cache outside of the editor
+		const TArray<FCultureRef> CurrentCultures = FInternationalization::Get().GetCurrentCultures(/*bIncludeLanguage*/true, /*bIncludeLocale*/false, /*bIncludeAssetGroups*/true);
+		for (const FCultureRef& CurrentCulture : CurrentCultures)
 		{
-			CultureCache->ConditionalUpdateCache();
+			if (TSharedPtr<FPackageLocalizationCultureCache> CultureCache = FindOrAddCacheForCulture_NoLock(CurrentCulture->GetName()))
+			{
+				CultureCache->ConditionalUpdateCache();
+			}
 		}
-	}
 
-	ConditionalUpdatePackageNameToAssetGroupCache_NoLock();
+		ConditionalUpdatePackageNameToAssetGroupCache_NoLock();
+	}
 }
diff --git a/Engine/Source/Runtime/CoreUObject/Private/Internationalization/PackageLocalizationManager.cpp b/Engine/Source/Runtime/CoreUObject/Private/Internationalization/PackageLocalizationManager.cpp
--- a/Engine/Source/Runtime/CoreUObject/Private/Internationalization/PackageLocalizationManager.cpp
+++ b/Engine/Source/Runtime/CoreUObject/Private/Internationalization/PackageLocalizationManager.cpp
@@ -84,10 +84,6 @@ void FPackageLocalizationManager::InitializeFromCache(const TSharedRef<IPackageL
 	if (!GIsEditor)
 	{
 		ActiveCache->ConditionalUpdateCache();
-
-		// Allow the plugin manager to update the package localization cache by exposing access through a delegate.
-		// PluginManager is a Core class, but package localization functionality is added at the CoreUObject level.
-		IPluginManager::Get().SetUpdatePackageLocalizationCacheDelegate(IPluginManager::FUpdatePackageLocalizationCacheDelegate::CreateRaw(this, &FPackageLocalizationManager::ConditionalUpdateCache));
 	}
 }
 
diff --git a/Engine/Source/Runtime/CoreUObject/Private/UObject/SavePackage2.cpp b/Engine/Source/Runtime/CoreUObject/Private/UObject/SavePackage2.cpp
--- a/Engine/Source/Runtime/CoreUObject/Private/UObject/SavePackage2.cpp
+++ b/Engine/Source/Runtime/CoreUObject/Private/UObject/SavePackage2.cpp
@@ -251,9 +251,6 @@ FORCEINLINE void EnsurePackageLocalization(UPackage* InPackage)
 		TextNamespaceUtil::EnsurePackageNamespace(InPackage);
 	}
 #endif // USE_STABLE_LOCALIZATION_KEYS
-
-	// Also make sure the localization cache is up to date, since updating it during the GIsSavingPackage won't allow object resolving
-	FPackageLocalizationManager::Get().ConditionalUpdateCache();
 }
 
 void PreSavePackage(FSaveContext& SaveContext)
diff --git a/Engine/Source/Runtime/Engine/Private/AssetManager.cpp b/Engine/Source/Runtime/Engine/Private/AssetManager.cpp
--- a/Engine/Source/Runtime/Engine/Private/AssetManager.cpp
+++ b/Engine/Source/Runtime/Engine/Private/AssetManager.cpp
@@ -2627,7 +2627,21 @@ bool UAssetManager::OnAssetRegistryAvailableAfterInitialization(FName InName, FA
 		if (bLoaded)
 		{
 			LocalAssetRegistry.AppendState(OutNewState);
-			FPackageLocalizationManager::Get().ConditionalUpdateCache();
+
+			// Invalidate the package localization cache for any updated mount roots
+			{
+				TArray<FString> RootPathsToInvalidate;
+				OutNewState.EnumerateAllPaths([&RootPathsToInvalidate](const FName PackagePath)
+				{
+					FString RootPath = FPackageName::SplitPackageNameRoot(PackagePath, nullptr);
+					RootPath.InsertAt(0, TEXT('/'));
+					RootPathsToInvalidate.AddUnique(MoveTemp(RootPath));
+				});
+				for (const FString& RootPathToInvalidate : RootPathsToInvalidate)
+				{
+					FPackageLocalizationManager::Get().InvalidateRootSourcePath(RootPathToInvalidate);
+				}
+			}
 
 			TArray<FAssetData> NewAssetData;
 			bool bRebuildReferenceList = false;
diff --git a/Engine/Source/Runtime/Engine/Private/Internationalization/EnginePackageLocalizationCache.cpp b/Engine/Source/Runtime/Engine/Private/Internationalization/EnginePackageLocalizationCache.cpp
--- a/Engine/Source/Runtime/Engine/Private/Internationalization/EnginePackageLocalizationCache.cpp
+++ b/Engine/Source/Runtime/Engine/Private/Internationalization/EnginePackageLocalizationCache.cpp
@@ -10,8 +10,8 @@
 FEnginePackageLocalizationCache::FEnginePackageLocalizationCache()
 	: bIsScanningPath(false)
 {
-	FAssetRegistryModule& AssetRegistryModule = FModuleManager::LoadModuleChecked<FAssetRegistryModule>(TEXT("AssetRegistry"));
-	IAssetRegistry& AssetRegistry = AssetRegistryModule.Get();
+	FModuleManager::LoadModuleChecked<FAssetRegistryModule>(TEXT("AssetRegistry"));
+	IAssetRegistry& AssetRegistry = IAssetRegistry::GetChecked();
 
 	AssetRegistry.OnAssetAdded().AddRaw(this, &FEnginePackageLocalizationCache::HandleAssetAdded);
 	AssetRegistry.OnAssetRemoved().AddRaw(this, &FEnginePackageLocalizationCache::HandleAssetRemoved);
@@ -20,16 +20,11 @@ FEnginePackageLocalizationCache::FEnginePackageLocalizationCache()
 
 FEnginePackageLocalizationCache::~FEnginePackageLocalizationCache()
 {
-	if (FModuleManager::Get().IsModuleLoaded(TEXT("AssetRegistry")))
+	if (IAssetRegistry* AssetRegistry = IAssetRegistry::Get())
 	{
-		FAssetRegistryModule& AssetRegistryModule = FModuleManager::GetModuleChecked<FAssetRegistryModule>(TEXT("AssetRegistry"));
-		IAssetRegistry* AssetRegistry = AssetRegistryModule.TryGet();
-		if (AssetRegistry)
-		{
-			AssetRegistry->OnAssetAdded().RemoveAll(this);
-			AssetRegistry->OnAssetRemoved().RemoveAll(this);
-			AssetRegistry->OnAssetRenamed().RemoveAll(this);
-		}
+		AssetRegistry->OnAssetAdded().RemoveAll(this);
+		AssetRegistry->OnAssetRemoved().RemoveAll(this);
+		AssetRegistry->OnAssetRenamed().RemoveAll(this);
 	}
 }
 
@@ -40,8 +35,7 @@ void FEnginePackageLocalizationCache::FindLocalizedPackages(const TMap<FString,
 		return;
 	}
 
-	FAssetRegistryModule& AssetRegistryModule = FModuleManager::LoadModuleChecked<FAssetRegistryModule>(TEXT("AssetRegistry"));
-	IAssetRegistry& AssetRegistry = AssetRegistryModule.Get();
+	IAssetRegistry& AssetRegistry = IAssetRegistry::GetChecked();
 
 	int32 NumElementsPerKeyGuess = NewSourceToLocalizedPaths.CreateConstIterator()->Value.Num();
 	int32 SizeGuess = NewSourceToLocalizedPaths.Num() * NumElementsPerKeyGuess;
@@ -89,8 +83,7 @@ void FEnginePackageLocalizationCache::FindLocalizedPackages(const TMap<FString,
 
 void FEnginePackageLocalizationCache::FindAssetGroupPackages(const FName InAssetGroupName, const FTopLevelAssetPath& InAssetClassName)
 {
-	FAssetRegistryModule& AssetRegistryModule = FModuleManager::LoadModuleChecked<FAssetRegistryModule>(TEXT("AssetRegistry"));
-	IAssetRegistry& AssetRegistry = AssetRegistryModule.Get();
+	IAssetRegistry& AssetRegistry = IAssetRegistry::GetChecked();
 
 	// We use the localized paths to find the source assets for the group since it's much faster to scan those paths than perform a full scan
 	TArray<FString> LocalizedRootPaths;
@@ -115,7 +108,7 @@ void FEnginePackageLocalizationCache::FindAssetGroupPackages(const FName InAsset
 	{
 		Filter.PackagePaths.Add(*LocalizedRootPath);
 	}
-	Filter.bIncludeOnlyOnDiskAssets = false;
+	Filter.bIncludeOnlyOnDiskAssets = !GIsEditor;
 	Filter.ClassPaths.Add(InAssetClassName);
 	Filter.bRecursiveClasses = false;
 
diff --git a/Engine/Source/Runtime/Projects/Private/PluginManager.cpp b/Engine/Source/Runtime/Projects/Private/PluginManager.cpp
--- a/Engine/Source/Runtime/Projects/Private/PluginManager.cpp
+++ b/Engine/Source/Runtime/Projects/Private/PluginManager.cpp
@@ -2932,11 +2932,6 @@ void FPluginManager::SetUnRegisterMountPointDelegate( const FRegisterMountPointD
 	UnRegisterMountPointDelegate = Delegate;
 }
 
-void FPluginManager::SetUpdatePackageLocalizationCacheDelegate( const FUpdatePackageLocalizationCacheDelegate& Delegate )
-{
-	UpdatePackageLocalizationCacheDelegate = Delegate;
-}
-
 bool FPluginManager::AreRequiredPluginsAvailable()
 {
 	return ConfigureEnabledPlugins();
@@ -3423,9 +3418,6 @@ void FPluginManager::MountPluginFromExternalSource(const TSharedRef<FPlugin>& Pl
 			{
 				EngineConfigFile->AddUniqueToSection(TEXT("Core.System"), "Paths", MoveTemp(ContentDir));
 			}
-
-			// Update the localization cache for the newly added content directory
-			UpdatePackageLocalizationCacheDelegate.ExecuteIfBound();
 		}
 	}
 
diff --git a/Engine/Source/Runtime/Projects/Private/PluginManager.h b/Engine/Source/Runtime/Projects/Private/PluginManager.h
--- a/Engine/Source/Runtime/Projects/Private/PluginManager.h
+++ b/Engine/Source/Runtime/Projects/Private/PluginManager.h
@@ -159,7 +159,6 @@ public:
 	virtual void GetLocalizationPathsForEnabledPlugins( TArray<FString>& OutLocResPaths ) override;
 	virtual void SetRegisterMountPointDelegate( const FRegisterMountPointDelegate& Delegate ) override;
 	virtual void SetUnRegisterMountPointDelegate( const FRegisterMountPointDelegate& Delegate ) override;
-	virtual void SetUpdatePackageLocalizationCacheDelegate( const FUpdatePackageLocalizationCacheDelegate& Delegate ) override;
 	virtual bool AreRequiredPluginsAvailable() override;
 	virtual FGetExplanationForUnavailablePackageWithPluginInfoDelegate& GetExplanationForUnavailablePackageWithPluginInfoDelegate() override;
 
@@ -350,10 +349,6 @@ private:
 	    content path unmounting functionality from Core. */
 	FRegisterMountPointDelegate UnRegisterMountPointDelegate;
 
-	/** Delegate for updating the package localization cache.  Bound by FPackageLocalizationManager code in 
-		CoreUObject, so that we can access localization cache functionality from Core. */
-	FUpdatePackageLocalizationCacheDelegate UpdatePackageLocalizationCacheDelegate;
-
 	/** Set if all the required plugins are available */
 	bool bHaveAllRequiredPlugins = false;
 
diff --git a/Engine/Source/Runtime/Projects/Public/Interfaces/IPluginManager.h b/Engine/Source/Runtime/Projects/Public/Interfaces/IPluginManager.h
--- a/Engine/Source/Runtime/Projects/Public/Interfaces/IPluginManager.h
+++ b/Engine/Source/Runtime/Projects/Public/Interfaces/IPluginManager.h
@@ -293,17 +293,6 @@ public:
 	 */
 	virtual void SetUnRegisterMountPointDelegate( const FRegisterMountPointDelegate& Delegate ) = 0;
 
-	/** Delegate type for updating the package localization cache.  Used internally by FPackageLocalizationManager code. */
-	DECLARE_DELEGATE( FUpdatePackageLocalizationCacheDelegate );
-
-	/**
-	 * Sets the delegate to call to update the package localization cache.  This is used internally by the plug-in manager system
-	 * and should not be called by you.  This is registered at application startup by FPackageLocalizationManager code in CoreUObject.
-	 *
-	 * @param	Delegate	The delegate to that will be called when plug-in manager needs to update the package localization cache
-	 */
-	virtual void SetUpdatePackageLocalizationCacheDelegate( const FUpdatePackageLocalizationCacheDelegate& Delegate ) = 0;
-
 	/**
 	 * Checks if all the required plug-ins are available. If not, will present an error dialog the first time a plug-in is loaded or this function is called.
 	 *
