Separate GFP PluginDetails and PluginURL query functionality.

#rnx

diff --git a/Engine/Plugins/Experimental/PluginAudit/Source/Private/SPluginAuditBrowser.cpp b/Engine/Plugins/Experimental/PluginAudit/Source/Private/SPluginAuditBrowser.cpp
--- a/Engine/Plugins/Experimental/PluginAudit/Source/Private/SPluginAuditBrowser.cpp
+++ b/Engine/Plugins/Experimental/PluginAudit/Source/Private/SPluginAuditBrowser.cpp
@@ -159,9 +159,8 @@ void SPluginAuditBrowser::BuildPluginList()
 	ExcludedGameFeaturePlugins.Reset();
 	for (const TSharedRef<IPlugin>& Plugin : AllPlugins)
 	{
-		FString PluginURL;
 		FGameFeaturePluginDetails PluginDetails;
-		if (GameFeaturesSubsystem.GetBuiltInGameFeaturePluginDetails(Plugin, PluginURL, PluginDetails))
+		if (GameFeaturesSubsystem.GetBuiltInGameFeaturePluginDetails(Plugin, PluginDetails))
 		{
 			if (Policy.WillPluginBeCooked(Plugin->GetDescriptorFileName(), PluginDetails))
 			{
diff --git a/Engine/Plugins/Runtime/GameFeatures/Source/GameFeatures/Private/GameFeaturesSubsystem.cpp b/Engine/Plugins/Runtime/GameFeatures/Source/GameFeatures/Private/GameFeaturesSubsystem.cpp
--- a/Engine/Plugins/Runtime/GameFeatures/Source/GameFeatures/Private/GameFeaturesSubsystem.cpp
+++ b/Engine/Plugins/Runtime/GameFeatures/Source/GameFeatures/Private/GameFeaturesSubsystem.cpp
@@ -1781,7 +1781,7 @@ void UGameFeaturesSubsystem::LoadBuiltInGameFeaturePlugin(const TSharedRef<IPlug
 
 	FString PluginURL;
 	FGameFeaturePluginDetails PluginDetails;
-	if (GetBuiltInGameFeaturePluginDetails(Plugin, PluginURL, PluginDetails))
+	if (GetBuiltInGameFeaturePluginDetails(Plugin, PluginDetails) && GetBuiltInGameFeaturePluginURL(Plugin, PluginURL))
 	{
 		if (IsPluginAllowed(PluginURL))
 		{
@@ -2102,10 +2102,31 @@ EGameFeaturePluginState UGameFeaturesSubsystem::GetPluginState(FGameFeaturePlugi
 
 bool UGameFeaturesSubsystem::GetGameFeaturePluginDetails(const TSharedRef<IPlugin>& Plugin, FString& OutPluginURL, FGameFeaturePluginDetails& OutPluginDetails) const
 {
-	return GetBuiltInGameFeaturePluginDetails(Plugin, OutPluginURL, OutPluginDetails);
+	return GetBuiltInGameFeaturePluginURL(Plugin, OutPluginURL) && GetBuiltInGameFeaturePluginDetails(Plugin, OutPluginDetails);
 }
 
 bool UGameFeaturesSubsystem::GetBuiltInGameFeaturePluginDetails(const TSharedRef<IPlugin>& Plugin, FString& OutPluginURL, FGameFeaturePluginDetails& OutPluginDetails) const
+{
+	return GetBuiltInGameFeaturePluginURL(Plugin, OutPluginURL) && GetBuiltInGameFeaturePluginDetails(Plugin, OutPluginDetails);
+}
+
+bool UGameFeaturesSubsystem::GetBuiltInGameFeaturePluginDetails(const TSharedRef<IPlugin>& Plugin, struct FGameFeaturePluginDetails& OutPluginDetails) const
+{
+	// @TODO: this problematic because it assumes file protocol.
+	// Ideally this would work with any protocol, but for current uses cases the exact protocol doesn't seem to matter.
+
+	const FString& PluginDescriptorFilename = Plugin->GetDescriptorFileName();
+	// Make sure you are in a game feature plugins folder. All GameFeaturePlugins are rooted in a GameFeatures folder.
+	if (!PluginDescriptorFilename.IsEmpty() && GetDefault<UGameFeaturesSubsystemSettings>()->IsValidGameFeaturePlugin(FPaths::ConvertRelativePathToFull(PluginDescriptorFilename)) && FPaths::FileExists(PluginDescriptorFilename))
+	{
+		return GetGameFeaturePluginDetailsInternal(PluginDescriptorFilename, OutPluginDetails);
+	}
+
+	return false;
+}
+
+
+bool UGameFeaturesSubsystem::GetBuiltInGameFeaturePluginURL(const TSharedRef<IPlugin>& Plugin, FString& OutPluginURL) const
 {
 	// @TODO: this problematic because it assumes file protocol.
 	// Ideally this would work with any protocol, but for current uses cases the exact protocol doesn't seem to matter.
@@ -2120,7 +2141,7 @@ bool UGameFeaturesSubsystem::GetBuiltInGameFeaturePluginDetails(const TSharedRef
 		{
 			GameSpecificPolicies->GetGameFeaturePluginURL(Plugin, OutPluginURL);
 		}
-		return GetGameFeaturePluginDetailsInternal(PluginDescriptorFilename, OutPluginDetails);
+		return true;
 	}
 
 	return false;
diff --git a/Engine/Plugins/Runtime/GameFeatures/Source/GameFeatures/Public/GameFeaturesSubsystem.h b/Engine/Plugins/Runtime/GameFeatures/Source/GameFeatures/Public/GameFeaturesSubsystem.h
--- a/Engine/Plugins/Runtime/GameFeatures/Source/GameFeatures/Public/GameFeaturesSubsystem.h
+++ b/Engine/Plugins/Runtime/GameFeatures/Source/GameFeatures/Public/GameFeaturesSubsystem.h
@@ -610,8 +610,15 @@ public:
 	bool GetGameFeaturePluginDetails(const TSharedRef<IPlugin>& Plugin, FString& OutPluginURL, struct FGameFeaturePluginDetails& OutPluginDetails) const;
 
 	/** Gets relevant properties out of a uplugin file. Should only be used for built-in GFPs */
+	UE_DEPRECATED(5.5, "Use non-PluginURL version of GetBuiltInGameFeaturePluginDetails and GetBuiltInGameFeaturePluginPath instead")
 	bool GetBuiltInGameFeaturePluginDetails(const TSharedRef<IPlugin>& Plugin, FString& OutPluginURL, struct FGameFeaturePluginDetails& OutPluginDetails) const;
 
+	/** Gets relevant properties out of a uplugin file. Should only be used for built-in GFPs */
+	bool GetBuiltInGameFeaturePluginDetails(const TSharedRef<IPlugin>& Plugin, struct FGameFeaturePluginDetails& OutPluginDetails) const;
+
+	/** Gets the URL for the given plugin, applying game-specific policies where appropriate. Should only be used for built-in GFPs */
+	bool GetBuiltInGameFeaturePluginURL(const TSharedRef<IPlugin>& Plugin, FString& OutPluginURL) const;
+
 	/** Gets relevant properties out of a uplugin file if it's installed */
 	bool GetGameFeaturePluginDetails(FString PluginURL, struct FGameFeaturePluginDetails& OutPluginDetails) const;
 
diff --git a/Engine/Plugins/Runtime/GameFeatures/Source/GameFeaturesEditor/Private/IllegalPluginDependenciesValidator.cpp b/Engine/Plugins/Runtime/GameFeatures/Source/GameFeaturesEditor/Private/IllegalPluginDependenciesValidator.cpp
--- a/Engine/Plugins/Runtime/GameFeatures/Source/GameFeaturesEditor/Private/IllegalPluginDependenciesValidator.cpp
+++ b/Engine/Plugins/Runtime/GameFeatures/Source/GameFeaturesEditor/Private/IllegalPluginDependenciesValidator.cpp
@@ -60,9 +60,8 @@ EDataValidationResult UIllegalPluginDependenciesValidator::ValidateLoadedAsset_I
 	GFPs.Reserve(AllPlugins.Num());
 	for (const TSharedRef<IPlugin>& Plugin : AllPlugins)
 	{
-		FString PluginURL;
 		FGameFeaturePluginDetails PluginDetails;
-		if (UGameFeaturesSubsystem::Get().GetBuiltInGameFeaturePluginDetails(Plugin, PluginURL, PluginDetails))
+		if (UGameFeaturesSubsystem::Get().GetBuiltInGameFeaturePluginDetails(Plugin, PluginDetails))
 		{
 			GFPs.Add(Plugin->GetName());
 		}
