Adds new API function for registering a GFP. Uses the range based test so that already activated GFPs will not be sent to registered.

diff --git a/Engine/Plugins/Runtime/GameFeatures/Source/GameFeatures/Private/GameFeatureData.cpp b/Engine/Plugins/Runtime/GameFeatures/Source/GameFeatures/Private/GameFeatureData.cpp
--- a/Engine/Plugins/Runtime/GameFeatures/Source/GameFeatures/Private/GameFeatureData.cpp
+++ b/Engine/Plugins/Runtime/GameFeatures/Source/GameFeatures/Private/GameFeatureData.cpp
@@ -757,6 +757,24 @@ FString UGameFeatureData::GetInstallBundleName(FStringView PluginName, bool bEve
 	}
 }
 
+FString UGameFeatureData::GetOptionalInstallBundleName(FStringView PluginName, bool bEvenIfDoesntExist /*= false*/)
+{
+	const FString OptionalBundleName = FString::Printf(TEXT("GFP_%.*sOptional"), PluginName.Len(), PluginName.GetData());
+	if (bEvenIfDoesntExist)
+	{
+		return OptionalBundleName;
+	}
+
+	if (InstallBundleUtil::HasInstallBundleInConfig(OptionalBundleName))
+	{
+		return OptionalBundleName;
+	}
+	else
+	{
+		return TEXT("");
+	}
+}
+
 void UGameFeatureData::GetPluginName(FString& PluginName) const
 {
 	UGameFeatureData::GetPluginName(this, PluginName);
diff --git a/Engine/Plugins/Runtime/GameFeatures/Source/GameFeatures/Private/GameFeaturesSubsystem.cpp b/Engine/Plugins/Runtime/GameFeatures/Source/GameFeatures/Private/GameFeaturesSubsystem.cpp
--- a/Engine/Plugins/Runtime/GameFeatures/Source/GameFeatures/Private/GameFeaturesSubsystem.cpp
+++ b/Engine/Plugins/Runtime/GameFeatures/Source/GameFeatures/Private/GameFeaturesSubsystem.cpp
@@ -1285,6 +1285,105 @@ void UGameFeaturesSubsystem::LoadGameFeaturePlugin(TConstArrayView<FString> Plug
 	}
 }
 
+void UGameFeaturesSubsystem::RegisterGameFeaturePlugin(const FString& PluginURL, const FGameFeaturePluginLoadComplete& CompleteDelegate)
+{
+	RegisterGameFeaturePlugin(PluginURL, FGameFeatureProtocolOptions(), CompleteDelegate);
+}
+
+void UGameFeaturesSubsystem::RegisterGameFeaturePlugin(const FString& PluginURL, const FGameFeatureProtocolOptions& ProtocolOptions, const FGameFeaturePluginLoadComplete& CompleteDelegate)
+{
+	const bool bIsPluginAllowed = IsPluginAllowed(PluginURL);
+	if (!bIsPluginAllowed)
+	{
+		CompleteDelegate.ExecuteIfBound(UE::GameFeatures::FResult(MakeError(UE::GameFeatures::SubsystemErrorNamespace + UE::GameFeatures::CommonErrorCodes::PluginNotAllowed)));
+		return;
+	}
+
+	UGameFeaturePluginStateMachine* StateMachine = FindOrCreateGameFeaturePluginStateMachine(PluginURL, ProtocolOptions);
+
+	if (!StateMachine->IsRunning() && StateMachine->GetCurrentState() == EGameFeaturePluginState::Active)
+	{
+		// TODO: Resolve the activated case here, this is needed because in a PIE environment the plugins
+		// are not sandboxed, and we need to do simulate a successful activate call in order run GFP systems 
+		// on whichever Role runs second between client and server.
+
+		// Refire the observer for Activated and do nothing else.
+		CallbackObservers(EObserverCallback::Activating, StateMachine->GetPluginIdentifier(), &StateMachine->GetPluginName(), StateMachine->GetGameFeatureDataForActivePlugin());
+	}
+
+	if (ShouldUpdatePluginProtocolOptions(StateMachine, ProtocolOptions))
+	{
+		const UE::GameFeatures::FResult Result = UpdateGameFeatureProtocolOptions(StateMachine, ProtocolOptions);
+		if (Result.HasError())
+		{
+			CompleteDelegate.ExecuteIfBound(Result);
+			return;
+		}
+	}
+
+	ChangeGameFeatureDestination(StateMachine, ProtocolOptions, FGameFeaturePluginStateRange(EGameFeaturePluginState::Registered, EGameFeaturePluginState::Active), CompleteDelegate);
+}
+
+void UGameFeaturesSubsystem::RegisterGameFeaturePlugin(TConstArrayView<FString> PluginURLs, const FGameFeatureProtocolOptions& ProtocolOptions, const FMultipleGameFeaturePluginsLoaded& CompleteDelegate)
+{
+	struct FLoadContext
+	{
+		TMap<FString, UE::GameFeatures::FResult> Results;
+		FMultipleGameFeaturePluginsLoaded CompleteDelegate;
+
+		int32 NumPluginsLoaded = 0;
+		bool bPushedTagsBroadcast = false;
+
+		FLoadContext()
+		{
+			if (!IsEngineExitRequested())
+			{
+				UGameplayTagsManager::Get().PushDeferOnGameplayTagTreeChangedBroadcast();
+				bPushedTagsBroadcast = true;
+			}
+			else if (UGameplayTagsManager* TagsManager = UGameplayTagsManager::GetIfAllocated())
+			{
+				TagsManager->PushDeferOnGameplayTagTreeChangedBroadcast();
+				bPushedTagsBroadcast = true;
+			}
+		}
+
+		~FLoadContext()
+		{
+			if (bPushedTagsBroadcast)
+			{
+				if (UGameplayTagsManager* TagsManager = UGameplayTagsManager::GetIfAllocated())
+				{
+					TagsManager->PopDeferOnGameplayTagTreeChangedBroadcast();
+				}
+			}
+
+			CompleteDelegate.ExecuteIfBound(Results);
+		}
+	};
+	TSharedRef<FLoadContext> LoadContext = MakeShared<FLoadContext>();
+	LoadContext->CompleteDelegate = CompleteDelegate;
+
+	LoadContext->Results.Reserve(PluginURLs.Num());
+	for (const FString& PluginURL : PluginURLs)
+	{
+		LoadContext->Results.Add(PluginURL, MakeError("Pending"));
+	}
+
+	const int32 NumPluginsToLoad = PluginURLs.Num();
+	UE_LOG(LogGameFeatures, Log, TEXT("Registering %i GFPs"), NumPluginsToLoad);
+
+	for (const FString& PluginURL : PluginURLs)
+	{
+		RegisterGameFeaturePlugin(PluginURL, ProtocolOptions, FGameFeaturePluginChangeStateComplete::CreateLambda([LoadContext, PluginURL](const UE::GameFeatures::FResult& Result)
+			{
+				LoadContext->Results.Add(PluginURL, Result);
+				++LoadContext->NumPluginsLoaded;
+				UE_LOG(LogGameFeatures, VeryVerbose, TEXT("Finished Registering %i GFPs"), LoadContext->NumPluginsLoaded);
+			}));
+	}
+}
+
 void UGameFeaturesSubsystem::LoadAndActivateGameFeaturePlugin(const FString& PluginURL, const FGameFeaturePluginLoadComplete& CompleteDelegate)
 {
 	ChangeGameFeatureTargetState(PluginURL, EGameFeatureTargetState::Active, CompleteDelegate);
@@ -1321,7 +1420,7 @@ void UGameFeaturesSubsystem::ChangeGameFeatureTargetState(const FString& PluginU
 	static_assert(std::underlying_type<EGameFeatureTargetState>::type(EGameFeatureTargetState::Count) == 4, "");
 	check(TargetPluginState != EGameFeaturePluginState::MAX);
 
-	const bool bIsPluginAllowed = IsPluginAllowed(PluginURL);
+		const bool bIsPluginAllowed = IsPluginAllowed(PluginURL);
 
 	UGameFeaturePluginStateMachine* StateMachine = nullptr;
 	if (!bIsPluginAllowed)
@@ -2665,7 +2764,7 @@ UGameFeaturePluginStateMachine* UGameFeaturesSubsystem::FindOrCreateGameFeatureP
 		// In this case, still return the existing machine, even if the protocol doesn't match. This function should never return null.
 		// There can only be one active instance of any machine.
 		EGameFeaturePluginProtocol ExpectedProtocol = (*ExistingStateMachine)->GetPluginIdentifier().GetPluginProtocol();
-		ensureAlwaysMsgf(ExpectedProtocol == PluginIdentifier.GetPluginProtocol(), TEXT("Expected protocol %s for %.*s"), UE::GameFeatures::GameFeaturePluginProtocolPrefix(ExpectedProtocol), PluginIdentifier.GetIdentifyingString().Len(), PluginIdentifier.GetIdentifyingString().GetData());
+		ensureMsgf(ExpectedProtocol == PluginIdentifier.GetPluginProtocol(), TEXT("Expected protocol %s for %.*s"), UE::GameFeatures::GameFeaturePluginProtocolPrefix(ExpectedProtocol), PluginIdentifier.GetIdentifyingString().Len(), PluginIdentifier.GetIdentifyingString().GetData());
 
 		UE_LOG(LogGameFeatures, VeryVerbose, TEXT("Found GameFeaturePlugin StateMachine using Identifier:%.*s from PluginURL:%s"), PluginIdentifier.GetIdentifyingString().Len(), PluginIdentifier.GetIdentifyingString().GetData(), *PluginURL);
 		return *ExistingStateMachine;
diff --git a/Engine/Plugins/Runtime/GameFeatures/Source/GameFeatures/Public/GameFeatureData.h b/Engine/Plugins/Runtime/GameFeatures/Source/GameFeatures/Public/GameFeatureData.h
--- a/Engine/Plugins/Runtime/GameFeatures/Source/GameFeatures/Public/GameFeatureData.h
+++ b/Engine/Plugins/Runtime/GameFeatures/Source/GameFeatures/Public/GameFeatureData.h
@@ -62,6 +62,13 @@ public:
 	 */
 	static FString GetInstallBundleName(FStringView PluginName, bool bEvenIfDoesntExist = false);
 
+	/**
+	 * Returns the optional install bundle name if one exists for this plugin.
+	 * @param - PluginName - the name of the GameFeaturePlugin we want to get a bundle for. Should be the same name as the .uplugin file
+	 * @param - bEvenIfDoesntExist - when true will return the name of bundle we are looking for without checking if it exists or not.
+	 */
+	static FString GetOptionalInstallBundleName(FStringView PluginName, bool bEvenIfDoesntExist = false);
+
 public:
 	//~UPrimaryDataAsset interface
 #if WITH_EDITORONLY_DATA
diff --git a/Engine/Plugins/Runtime/GameFeatures/Source/GameFeatures/Public/GameFeaturesSubsystem.h b/Engine/Plugins/Runtime/GameFeatures/Source/GameFeatures/Public/GameFeaturesSubsystem.h
--- a/Engine/Plugins/Runtime/GameFeatures/Source/GameFeatures/Public/GameFeaturesSubsystem.h
+++ b/Engine/Plugins/Runtime/GameFeatures/Source/GameFeatures/Public/GameFeaturesSubsystem.h
@@ -509,6 +509,11 @@ public:
 	void LoadGameFeaturePlugin(const FString& PluginURL, const FGameFeatureProtocolOptions& ProtocolOptions, const FGameFeaturePluginLoadComplete& CompleteDelegate);
 	void LoadGameFeaturePlugin(TConstArrayView<FString> PluginURLs, const FGameFeatureProtocolOptions& ProtocolOptions, const FMultipleGameFeaturePluginsLoaded& CompleteDelegate);
 
+	/** Registers a single game feature plugin. */
+	void RegisterGameFeaturePlugin(const FString& PluginURL, const FGameFeaturePluginLoadComplete& CompleteDelegate);
+	void RegisterGameFeaturePlugin(const FString& PluginURL, const FGameFeatureProtocolOptions& ProtocolOptions, const FGameFeaturePluginLoadComplete& CompleteDelegate);
+	void RegisterGameFeaturePlugin(TConstArrayView<FString> PluginURLs, const FGameFeatureProtocolOptions& ProtocolOptions, const FMultipleGameFeaturePluginsLoaded& CompleteDelegate);
+
 	/** Loads a single game feature plugin and activates it. */
 	void LoadAndActivateGameFeaturePlugin(const FString& PluginURL, const FGameFeaturePluginLoadComplete& CompleteDelegate);
 	void LoadAndActivateGameFeaturePlugin(const FString& PluginURL, const FGameFeatureProtocolOptions& ProtocolOptions, const FGameFeaturePluginLoadComplete& CompleteDelegate);
