Cleanup stale TODO comment.

diff --git a/Engine/Plugins/Runtime/GameFeatures/Source/GameFeatures/Private/GameFeaturesSubsystem.cpp b/Engine/Plugins/Runtime/GameFeatures/Source/GameFeatures/Private/GameFeaturesSubsystem.cpp
--- a/Engine/Plugins/Runtime/GameFeatures/Source/GameFeatures/Private/GameFeaturesSubsystem.cpp
+++ b/Engine/Plugins/Runtime/GameFeatures/Source/GameFeatures/Private/GameFeaturesSubsystem.cpp
@@ -1593,8 +1593,6 @@ bool UGameFeaturesSubsystem::ShouldUpdatePluginProtocolOptions(const UGameFeatur
 
 bool UGameFeaturesSubsystem::GetGameFeaturePluginInstallPercent(const FString& PluginURL, float& Install_Percent) const
 {
-	// TODO: figure out if doing IAD streaming and include IAD progress if needed
-
 	if (const UGameFeaturePluginStateMachine* StateMachine = FindGameFeaturePluginStateMachine(PluginURL))
 	{
 		if (StateMachine->IsStatusKnown() && StateMachine->IsAvailable())
