Optimize UGameFeatureSetSubsystem::GetDiscoveredFeatureSetPlugins, especially for Zen file server

 - Parallelize the main loop
 - Make a couple of downstream functions thread-safe
 - Added a version of GetAllowedFeatureSetPlugins that takes an existing list of discovered game feature set plugins, to avoid needing to call GetDiscoveredFeatureSetPlugins twice.

#rnx

diff --git a/Engine/Plugins/Runtime/GameFeatures/Source/GameFeatures/Private/GameFeaturesSubsystem.cpp b/Engine/Plugins/Runtime/GameFeatures/Source/GameFeatures/Private/GameFeaturesSubsystem.cpp
--- a/Engine/Plugins/Runtime/GameFeatures/Source/GameFeatures/Private/GameFeaturesSubsystem.cpp
+++ b/Engine/Plugins/Runtime/GameFeatures/Source/GameFeatures/Private/GameFeaturesSubsystem.cpp
@@ -2265,6 +2265,8 @@ bool UGameFeaturesSubsystem::GetGameFeaturePluginDetailsInternal(const FString&
 	FDateTime FileTimeStamp;
 	if (UE::GameFeatures::GCachePluginDetails)
 	{
+		FScopeLock Lock(&CachedGameFeaturePluginDetailsLock);
+
 		// Note: On file systems that don't support timestamps (current time is returned instead), 
 		// the pak file layer will end up caching the mount time, so this stamp will still be valid as long as the uplugin
 		// is in a pak and the pak is mounted.
@@ -2386,6 +2388,7 @@ bool UGameFeaturesSubsystem::GetGameFeaturePluginDetailsInternal(const FString&
 
 	if (UE::GameFeatures::GCachePluginDetails)
 	{
+		FScopeLock Lock(&CachedGameFeaturePluginDetailsLock);
 		CachedPluginDetailsByFilename.Add(PluginDescriptorFilename, FCachedGameFeaturePluginDetails(OutPluginDetails, FileTimeStamp));
 	}
 
diff --git a/Engine/Plugins/Runtime/GameFeatures/Source/GameFeatures/Private/GameFeaturesSubsystemSettings.cpp b/Engine/Plugins/Runtime/GameFeatures/Source/GameFeatures/Private/GameFeaturesSubsystemSettings.cpp
--- a/Engine/Plugins/Runtime/GameFeatures/Source/GameFeatures/Private/GameFeaturesSubsystemSettings.cpp
+++ b/Engine/Plugins/Runtime/GameFeatures/Source/GameFeatures/Private/GameFeaturesSubsystemSettings.cpp
@@ -19,27 +19,32 @@ UGameFeaturesSubsystemSettings::UGameFeaturesSubsystemSettings()
 bool UGameFeaturesSubsystemSettings::IsValidGameFeaturePlugin(const FString& PluginDescriptorFilename) const
 {
 	// Build the cache of game feature plugin folders the first time this is called
-	if (BuiltInGameFeaturePluginsFolders.IsEmpty())
+	static struct FBuiltInGameFeaturePluginsFolders
 	{
-		// Get all the existing game feature paths
-		BuiltInGameFeaturePluginsFolders.Append(
-			FPaths::GetExtensionDirs(FPaths::ProjectDir(), FPaths::Combine(TEXT("Plugins"), TEXT("GameFeatures")))
-		);
-
-		// The base directory may not exist yet, add it if empty
-		if (BuiltInGameFeaturePluginsFolders.IsEmpty())
+		FBuiltInGameFeaturePluginsFolders()
 		{
-			BuiltInGameFeaturePluginsFolders.Add(FPaths::Combine(FPaths::ProjectDir(), TEXT("Plugins"), TEXT("GameFeatures")));
+			// Get all the existing game feature paths
+			BuiltInGameFeaturePluginsFolders.Append(
+				FPaths::GetExtensionDirs(FPaths::ProjectDir(), FPaths::Combine(TEXT("Plugins"), TEXT("GameFeatures")))
+			);
+
+			// The base directory may not exist yet, add it if empty
+			if (BuiltInGameFeaturePluginsFolders.IsEmpty())
+			{
+				BuiltInGameFeaturePluginsFolders.Add(FPaths::Combine(FPaths::ProjectDir(), TEXT("Plugins"), TEXT("GameFeatures")));
+			}
+
+			for (FString& BuiltInFolder : BuiltInGameFeaturePluginsFolders)
+			{
+				BuiltInFolder = FPaths::ConvertRelativePathToFull(BuiltInFolder + TEXT("/"));
+			}
 		}
 
-		for (FString& BuiltInFolder : BuiltInGameFeaturePluginsFolders)
-		{
-			BuiltInFolder = FPaths::ConvertRelativePathToFull(BuiltInFolder + TEXT("/"));
-		}
-	}
+		TArray<FString> BuiltInGameFeaturePluginsFolders;
+	} Lazy;
 
 	// Check to see if the filename is rooted in a game feature plugin folder
-	for (const FString& BuiltInFolder : BuiltInGameFeaturePluginsFolders)
+	for (const FString& BuiltInFolder : Lazy.BuiltInGameFeaturePluginsFolders)
 	{
 		if (FPathViews::IsParentPathOf(BuiltInFolder, PluginDescriptorFilename))
 		{
diff --git a/Engine/Plugins/Runtime/GameFeatures/Source/GameFeatures/Public/GameFeaturesSubsystem.h b/Engine/Plugins/Runtime/GameFeatures/Source/GameFeatures/Public/GameFeaturesSubsystem.h
--- a/Engine/Plugins/Runtime/GameFeatures/Source/GameFeatures/Public/GameFeaturesSubsystem.h
+++ b/Engine/Plugins/Runtime/GameFeatures/Source/GameFeatures/Public/GameFeaturesSubsystem.h
@@ -801,6 +801,7 @@ private:
 		FCachedGameFeaturePluginDetails(const FGameFeaturePluginDetails& InDetails, const FDateTime& InTimeStamp) : Details(InDetails), TimeStamp(InTimeStamp) {}
 	};
 	mutable TMap<FString, FCachedGameFeaturePluginDetails> CachedPluginDetailsByFilename;
+	mutable FCriticalSection CachedGameFeaturePluginDetailsLock;
 
 	UPROPERTY()
 	TArray<TObjectPtr<UObject>> Observers;
diff --git a/Engine/Plugins/Runtime/GameFeatures/Source/GameFeatures/Public/GameFeaturesSubsystemSettings.h b/Engine/Plugins/Runtime/GameFeatures/Source/GameFeatures/Public/GameFeaturesSubsystemSettings.h
--- a/Engine/Plugins/Runtime/GameFeatures/Source/GameFeatures/Public/GameFeaturesSubsystemSettings.h
+++ b/Engine/Plugins/Runtime/GameFeatures/Source/GameFeatures/Public/GameFeaturesSubsystemSettings.h
@@ -39,10 +39,6 @@ public:
 	UE_DEPRECATED(5.0, "Use IsValidGameFeaturePlugin() instead")
 	FString BuiltInGameFeaturePluginsFolder;
 
-private:
-	// Cached list of all GameFeature plugin directories (including extension versions)
-	mutable TArray<FString> BuiltInGameFeaturePluginsFolders;
-
 public:
 	// Returns true if the specified (fully qualified) path is a game feature plugin
 	bool IsValidGameFeaturePlugin(const FString& PluginDescriptorFilename) const;
