When determining protocol dependencies allow ini loading for the dependencies of file protocol GFPs.

diff --git a/Engine/Plugins/Runtime/GameFeatures/Source/GameFeatures/Private/GameFeaturesSubsystem.cpp b/Engine/Plugins/Runtime/GameFeatures/Source/GameFeatures/Private/GameFeaturesSubsystem.cpp
--- a/Engine/Plugins/Runtime/GameFeatures/Source/GameFeatures/Private/GameFeaturesSubsystem.cpp
+++ b/Engine/Plugins/Runtime/GameFeatures/Source/GameFeatures/Private/GameFeaturesSubsystem.cpp
@@ -2889,7 +2889,7 @@ bool UGameFeaturesSubsystem::FindOrCreatePluginDependencyStateMachines(const FSt
 {
 	const FString& PluginFilename = InStateProperties.PluginInstalledFilename;
 	const FGameFeatureProtocolOptions InDepProtocolOptions = InStateProperties.RecycleProtocolOptions();
-
+    const EGameFeaturePluginProtocol InProtocol = UGameFeaturesSubsystem::GetPluginURLProtocol(PluginURL);
 	const bool bWarnOnDepCreation = InStateProperties.ProtocolOptions.bLogWarningOnForcedDependencyCreation;
 	const bool bErrorOnDepCreation = InStateProperties.ProtocolOptions.bLogErrorOnForcedDependencyCreation;
 
@@ -2924,17 +2924,23 @@ bool UGameFeaturesSubsystem::FindOrCreatePluginDependencyStateMachines(const FSt
 			// Inherit dep protocol options if possible
 			FGameFeatureProtocolOptions DepProtocolOptions;
 			EGameFeaturePluginProtocol DepProtocol = UGameFeaturesSubsystem::GetPluginURLProtocol(DependencyURL);
-			if (DepProtocol == EGameFeaturePluginProtocol::InstallBundle && InDepProtocolOptions.HasSubtype<FInstallBundlePluginProtocolOptions>())
-			{
-				DepProtocolOptions = InDepProtocolOptions;
-			}
-			else
-			{
-				// Always propogate non-protocol specific flags
-				DepProtocolOptions.bForceSyncLoading = InDepProtocolOptions.bForceSyncLoading;
-				DepProtocolOptions.bLogWarningOnForcedDependencyCreation = InDepProtocolOptions.bLogWarningOnForcedDependencyCreation;
-				DepProtocolOptions.bLogErrorOnForcedDependencyCreation = InDepProtocolOptions.bLogErrorOnForcedDependencyCreation;
-			}
+			if (DepProtocol == EGameFeaturePluginProtocol::InstallBundle)
+            {
+                if (InDepProtocolOptions.HasSubtype<FInstallBundlePluginProtocolOptions>())
+                {
+                    DepProtocolOptions = InDepProtocolOptions;
+                }
+                else if (InProtocol == EGameFeaturePluginProtocol::File)
+                {
+                    FInstallBundlePluginProtocolOptions InstallBundleOptions;
+                    InstallBundleOptions.bAllowIniLoading = true;
+                    DepProtocolOptions = FGameFeatureProtocolOptions(InstallBundleOptions);
+                }
+            }
+            // Always propogate non-protocol specific flags
+            DepProtocolOptions.bForceSyncLoading = InDepProtocolOptions.bForceSyncLoading;
+            DepProtocolOptions.bLogWarningOnForcedDependencyCreation = InDepProtocolOptions.bLogWarningOnForcedDependencyCreation;
+            DepProtocolOptions.bLogErrorOnForcedDependencyCreation = InDepProtocolOptions.bLogErrorOnForcedDependencyCreation;
 
 			bool bFoundExisting = false;
 			UGameFeaturePluginStateMachine* ResolvedDependency = FindOrCreateGameFeaturePluginStateMachine(DependencyURL, DepProtocolOptions, &bFoundExisting);
