Improved the way that Data Registry GFA entries are displayed to show the name in the preview for each row.
#rnx

diff --git a/Engine/Plugins/Runtime/GameFeatures/Source/GameFeatures/Public/GameFeatureAction_DataRegistrySource.h b/Engine/Plugins/Runtime/GameFeatures/Source/GameFeatures/Public/GameFeatureAction_DataRegistrySource.h
--- a/Engine/Plugins/Runtime/GameFeatures/Source/GameFeatures/Public/GameFeatureAction_DataRegistrySource.h
+++ b/Engine/Plugins/Runtime/GameFeatures/Source/GameFeatures/Public/GameFeatureAction_DataRegistrySource.h
@@ -72,7 +72,7 @@ public:
 
 private:
 	/** List of sources to add when this feature is activated */
-	UPROPERTY(EditAnywhere, Category = "Registry Data")
+	UPROPERTY(EditAnywhere, Category = "Registry Data", meta=(TitleProperty="RegistryToAddTo"))
 	TArray<FDataRegistrySourceToAdd> SourcesToAdd;
 
 	/** If true, this will preload the sources when the feature is registered in the editor to support the editor pickers */
