Only when no other GFPs are using a bundle will it be unmounted.
Add new function GetRemovableBundles to determine which bundles are not being used.

diff --git a/Engine/Plugins/Runtime/GameFeatures/Source/GameFeatures/Private/GameFeaturePluginStateMachine.cpp b/Engine/Plugins/Runtime/GameFeatures/Source/GameFeatures/Private/GameFeaturePluginStateMachine.cpp
--- a/Engine/Plugins/Runtime/GameFeatures/Source/GameFeatures/Private/GameFeaturePluginStateMachine.cpp
+++ b/Engine/Plugins/Runtime/GameFeatures/Source/GameFeatures/Private/GameFeaturePluginStateMachine.cpp
@@ -1006,12 +1006,8 @@ struct FBaseDataReleaseGameFeaturePluginState : public FGameFeaturePluginState
 			return;
 		}
 
-		TSharedPtr<IInstallBundleManager> BundleManager = IInstallBundleManager::GetPlatformInstallBundleManager();
-		check(BundleManager.IsValid());
-
 		EInstallBundleReleaseRequestFlags ReleaseFlags = GetReleaseRequestFlags();
-		TValueOrError<FInstallBundleReleaseRequestInfo, EInstallBundleResult> MaybeRequestInfo = BundleManager->RequestReleaseContent(GetInstallBundles(), ReleaseFlags);
-
+		TValueOrError<FInstallBundleReleaseRequestInfo, EInstallBundleResult> MaybeRequestInfo = UGameFeaturesSubsystem::Get().ReleaseBundle(StateProperties.PluginName, GetInstallBundles(), ReleaseFlags);
 		if (MaybeRequestInfo.HasError())
 		{
 			const FStringView ShortUrl = StateProperties.PluginIdentifier.GetIdentifyingString();
@@ -1492,7 +1488,8 @@ struct FBaseDownloadGameFeaturePluginState : public FGameFeaturePluginState
 		{
 			return;
 		}
-
+		
+		UGameFeaturesSubsystem::Get().OnGameFeatureDownloaded(StateProperties.PluginIdentifier);
 		const EGameFeaturePluginState SuccessState = GetSuccessTransitionState();
 		StateStatus.SetTransition(SuccessState);
 	}
@@ -1747,16 +1744,13 @@ struct FGameFeaturePluginState_Unmounting : public FGameFeaturePluginState
 			return;
 		}
 
-		TSharedPtr<IInstallBundleManager> BundleManager = IInstallBundleManager::GetPlatformInstallBundleManager();
-
 		const TArray<FName>& InstallBundles = StateProperties.ProtocolMetadata.GetSubtype<FInstallBundlePluginProtocolMetaData>().InstallBundles;
 
 		EInstallBundleReleaseRequestFlags ReleaseFlags = StateProperties.ProtocolOptions.GetSubtype<FInstallBundlePluginProtocolOptions>().ReleaseInstallBundleFlags;
 		ReleaseFlags |= EInstallBundleReleaseRequestFlags::SkipReleaseUnmountOnly;
 		//Make sure we don't remove files here early, that should only be done in Uninstalling
 		ReleaseFlags &= ~(EInstallBundleReleaseRequestFlags::RemoveFilesIfPossible);
-
-		TValueOrError<FInstallBundleReleaseRequestInfo, EInstallBundleResult> MaybeRequestInfo = BundleManager->RequestReleaseContent(InstallBundles, ReleaseFlags);
+		TValueOrError<FInstallBundleReleaseRequestInfo, EInstallBundleResult> MaybeRequestInfo = UGameFeaturesSubsystem::Get().UnmountBundle(StateProperties.PluginName, InstallBundles, ReleaseFlags);
 
 		if (MaybeRequestInfo.HasError())
 		{
@@ -3549,6 +3543,11 @@ UGameFeatureData* UGameFeaturePluginStateMachine::GetGameFeatureDataForRegistere
 	return nullptr;
 }
 
+const FGameFeaturePluginStateMachineProperties& UGameFeaturePluginStateMachine::GetProperties() const
+{
+	return StateProperties;
+}
+
 bool UGameFeaturePluginStateMachine::IsValidTransitionState(EGameFeaturePluginState InState) const
 {
 	check(InState != EGameFeaturePluginState::MAX);
diff --git a/Engine/Plugins/Runtime/GameFeatures/Source/GameFeatures/Private/GameFeaturePluginStateMachine.h b/Engine/Plugins/Runtime/GameFeatures/Source/GameFeatures/Private/GameFeaturePluginStateMachine.h
--- a/Engine/Plugins/Runtime/GameFeatures/Source/GameFeatures/Private/GameFeaturePluginStateMachine.h
+++ b/Engine/Plugins/Runtime/GameFeatures/Source/GameFeatures/Private/GameFeaturePluginStateMachine.h
@@ -552,6 +552,8 @@ public:
 	/** If the plugin is registered already, we will retrieve its game feature data */
 	UGameFeatureData* GetGameFeatureDataForRegisteredPlugin(bool bCheckForRegistering = false);
 
+	const FGameFeaturePluginStateMachineProperties& GetProperties() const;
+
 private:
 	/** Returns true if the specified state is not a transition state */
 	bool IsValidTransitionState(EGameFeaturePluginState InState) const;
diff --git a/Engine/Plugins/Runtime/GameFeatures/Source/GameFeatures/Private/GameFeaturesSubsystem.cpp b/Engine/Plugins/Runtime/GameFeatures/Source/GameFeatures/Private/GameFeaturesSubsystem.cpp
--- a/Engine/Plugins/Runtime/GameFeatures/Source/GameFeatures/Private/GameFeaturesSubsystem.cpp
+++ b/Engine/Plugins/Runtime/GameFeatures/Source/GameFeatures/Private/GameFeaturesSubsystem.cpp
@@ -968,6 +968,23 @@ void UGameFeaturesSubsystem::OnGameFeatureDownloading(const FString& PluginName,
 	CallbackObservers(EObserverCallback::Downloading, PluginIdentifier, &PluginName);
 }
 
+void UGameFeaturesSubsystem::OnGameFeatureDownloaded(const FGameFeaturePluginIdentifier& PluginIdentifier)
+{
+	if (UGameFeaturePluginStateMachine* GFPSM = FindGameFeaturePluginStateMachine(PluginIdentifier))
+	{
+		const FGameFeaturePluginStateMachineProperties& MachineProperties = GFPSM->GetProperties();
+		if (MachineProperties.GetPluginProtocol() == EGameFeaturePluginProtocol::InstallBundle)
+		{
+			const FInstallBundlePluginProtocolMetaData& MetaData = MachineProperties.ProtocolMetadata.GetSubtype<FInstallBundlePluginProtocolMetaData>();
+			for (const FName& Bundle : MetaData.InstallBundles)
+			{
+				TSet<FString>& PluginsUsingBundle = DownloadedBundleToPlugin.FindOrAdd(Bundle);
+				PluginsUsingBundle.Add(FString(PluginIdentifier.GetPluginName()));
+			}
+		}
+	}
+}
+
 void UGameFeaturesSubsystem::OnGameFeatureReleasing(const FString& PluginName, const FGameFeaturePluginIdentifier& PluginIdentifier)
 {
 	CallbackObservers(EObserverCallback::Releasing, PluginIdentifier, &PluginName);
@@ -980,6 +997,20 @@ void UGameFeaturesSubsystem::OnGameFeaturePreMounting(const FString& PluginName,
 
 void UGameFeaturesSubsystem::OnGameFeaturePostMounting(const FString& PluginName, const FGameFeaturePluginIdentifier& PluginIdentifier, FGameFeaturePostMountingContext& Context)
 {
+	if (UGameFeaturePluginStateMachine* GFPSM = FindGameFeaturePluginStateMachine(PluginIdentifier))
+	{
+		const FGameFeaturePluginStateMachineProperties& MachineProperties = GFPSM->GetProperties();
+		if (MachineProperties.GetPluginProtocol() == EGameFeaturePluginProtocol::InstallBundle)
+		{
+			const FInstallBundlePluginProtocolMetaData& MetaData = MachineProperties.ProtocolMetadata.GetSubtype<FInstallBundlePluginProtocolMetaData>();
+			for (const FName& Bundle : MetaData.InstallBundles)
+			{
+				TSet<FString>& PluginsUsingBundle = MountedBundleToPlugin.FindOrAdd(Bundle);
+				PluginsUsingBundle.Add(PluginName);
+			}
+		}
+	}
+
 	CallbackObservers(EObserverCallback::PostMounting, PluginIdentifier, &PluginName, /*GameFeatureData=*/nullptr, &Context);
 }
 
@@ -1891,6 +1922,7 @@ void UGameFeaturesSubsystem::LoadBuiltInGameFeaturePlugin(const TSharedRef<IPlug
 				{
 					FInstallBundlePluginProtocolOptions InstallBundleOptions;
 					InstallBundleOptions.bAllowIniLoading = true;
+					InstallBundleOptions.ReleaseInstallBundleFlags = EInstallBundleReleaseRequestFlags::ExplicitRemoveList;
 					ProtocolOptions = FGameFeatureProtocolOptions(InstallBundleOptions);
 				}
 				ProtocolOptions.bForceSyncLoading = BehaviorOptions.bForceSyncLoading;
@@ -2400,6 +2432,47 @@ void UGameFeaturesSubsystem::PruneCachedGameFeaturePluginDetails(const FString&
 	CachedPluginDetailsByFilename.Remove(PluginDescriptorFilename);
 }
 
+TValueOrError<FInstallBundleReleaseRequestInfo, EInstallBundleResult> UGameFeaturesSubsystem::ReleaseBundle(const FString& PluginName, TArrayView<const FName> BundleNames, EInstallBundleReleaseRequestFlags Flags)
+{
+	TArray<FName> BundlesToRelease;
+	for (const FName& Bundle : BundleNames)
+	{
+		TSet<FString>& PluginsUsingBundle = DownloadedBundleToPlugin.FindOrAdd(Bundle);
+		PluginsUsingBundle.Remove(PluginName);
+		if (PluginsUsingBundle.Num() == 0)
+		{
+			UE_LOG(LogGameFeatures, Display, TEXT("Releasing bundle (%s) as final plugin released it (%s)"), *Bundle.ToString(), *PluginName);
+			BundlesToRelease.Add(Bundle);
+		}
+
+#if !UE_BUILD_SHIPPING
+		TSet<FString>& PluginsWithMountedBundle = MountedBundleToPlugin.FindOrAdd(Bundle);
+		ensureMsgf(!PluginsWithMountedBundle.Contains(PluginName), TEXT("Bundle (%s) is being released but still has GFPs with it mounted"), *Bundle.ToString());
+#endif
+	}
+	TSharedPtr<IInstallBundleManager> BundleManager = IInstallBundleManager::GetPlatformInstallBundleManager();
+	check(BundleManager.IsValid());
+	return BundleManager->RequestReleaseContent(BundlesToRelease, Flags);
+}
+
+TValueOrError<FInstallBundleReleaseRequestInfo, EInstallBundleResult> UGameFeaturesSubsystem::UnmountBundle(const FString& PluginName, TArrayView<const FName> BundleNames, EInstallBundleReleaseRequestFlags Flags)
+{
+	TArray<FName> BundlesToRelease;
+	for (const FName& Bundle : BundleNames)
+	{
+		TSet<FString>& PluginsUsingBundle = MountedBundleToPlugin.FindOrAdd(Bundle);
+		PluginsUsingBundle.Remove(PluginName);
+		if (PluginsUsingBundle.Num() == 0)
+		{
+			UE_LOG(LogGameFeatures, Display, TEXT("Unmounting bundle (%s) as final plugin released it (%s)"), *Bundle.ToString(), *PluginName);
+			BundlesToRelease.Add(Bundle);
+		}
+	}
+	TSharedPtr<IInstallBundleManager> BundleManager = IInstallBundleManager::GetPlatformInstallBundleManager();
+	check(BundleManager.IsValid());
+	return BundleManager->RequestReleaseContent(BundlesToRelease, Flags);
+}
+
 struct FGameFeaturePluginPredownloadContext : public FGameFeaturePluginPredownloadHandle
 {
 	const FStringView PredownloadErrorNamespace = TEXTVIEW("GameFeaturePlugin.Predownload.");
@@ -3458,4 +3531,3 @@ void UGameFeaturesSubsystem::SetExplanationForNotMountingPlugin(const FString& P
 	}
 #endif
 }
-
diff --git a/Engine/Plugins/Runtime/GameFeatures/Source/GameFeatures/Public/GameFeaturesSubsystem.h b/Engine/Plugins/Runtime/GameFeatures/Source/GameFeatures/Public/GameFeaturesSubsystem.h
--- a/Engine/Plugins/Runtime/GameFeatures/Source/GameFeatures/Public/GameFeaturesSubsystem.h
+++ b/Engine/Plugins/Runtime/GameFeatures/Source/GameFeatures/Public/GameFeaturesSubsystem.h
@@ -21,6 +21,8 @@ class FJsonObject;
 struct FWorldContext;
 struct FGameFeaturePluginStateRange;
 struct FGameFeaturePluginStateMachineProperties;
+struct FInstallBundleReleaseRequestInfo;
+enum class EInstallBundleResult : uint32;
 enum class EInstallBundleRequestFlags : uint32;
 enum class EInstallBundleReleaseRequestFlags : uint32;
 
@@ -656,6 +658,7 @@ private:
 	void OnGameFeaturePredownloading(const FString& PluginName, const FGameFeaturePluginIdentifier& PluginIdentifier);
 
 	void OnGameFeatureDownloading(const FString& PluginName, const FGameFeaturePluginIdentifier& PluginIdentifier);
+	void OnGameFeatureDownloaded(const FGameFeaturePluginIdentifier& PluginIdentifier);
 	friend struct FGameFeaturePluginState_Downloading;
 
 	void OnGameFeatureReleasing(const FString& PluginName, const FGameFeaturePluginIdentifier& PluginIdentifier);
@@ -712,6 +715,23 @@ private:
 	/** Prunes any cached GFP details */
 	void PruneCachedGameFeaturePluginDetails(const FString& PluginURL, const FString& PluginDescriptorFilename) const;
 	friend struct FGameFeaturePluginState_Unmounting;
+	friend struct FBaseDataReleaseGameFeaturePluginState;
+
+	/**
+	 * Wrapper to InstallBundleManager release bundle.
+	 * Will only release the bundle if this is the last GFP using the bundle.
+	 */
+	TValueOrError<FInstallBundleReleaseRequestInfo, EInstallBundleResult> ReleaseBundle(const FString& PluginName, TArrayView<const FName> BundleNames, EInstallBundleReleaseRequestFlags Flags);
+
+	/**
+	 * Wrapper to InstallBundleManager release bundle for the unmount call
+	 * Will only unmount the bundle if this is the last GFP using the bundle.
+	 */
+	TValueOrError<FInstallBundleReleaseRequestInfo, EInstallBundleResult> UnmountBundle(const FString& PluginName,  TArrayView<const FName> BundleNames, EInstallBundleReleaseRequestFlags Flags);
+
+	// Reference tracker to know which bundles are in use.
+	TMap<FName, TSet<FString>> MountedBundleToPlugin;
+	TMap<FName, TSet<FString>> DownloadedBundleToPlugin;
 
 	/** Gets the state machine associated with the specified URL */
 	UGameFeaturePluginStateMachine* FindGameFeaturePluginStateMachine(const FString& PluginURL) const;
