API changes to allow the Core Content bundle to be updated and the install status queried. Added additional bundle state to track if a VerseURI is "Up To Date". Core Content update will be automatically kicked off when entering the frontend if required.

diff --git a/Engine/Plugins/Runtime/GameFeatures/Source/GameFeatures/Private/GameFeaturesSubsystem.cpp b/Engine/Plugins/Runtime/GameFeatures/Source/GameFeatures/Private/GameFeaturesSubsystem.cpp
--- a/Engine/Plugins/Runtime/GameFeatures/Source/GameFeatures/Private/GameFeaturesSubsystem.cpp
+++ b/Engine/Plugins/Runtime/GameFeatures/Source/GameFeatures/Private/GameFeaturesSubsystem.cpp
@@ -1690,6 +1690,32 @@ bool UGameFeaturesSubsystem::GetGameFeaturePluginInstallPercent(TConstArrayView<
 	return false;
 }
 
+bool UGameFeaturesSubsystem::IsGameFeaturePluginUpToDate(const FString& PluginURL) const
+{
+	TArray<FName> InstallBundles;
+	const bool bParseSuccess = UGameFeaturesSubsystem::ParsePluginURLOptions(PluginURL, EGameFeatureURLOptions::Bundles,
+		[&InstallBundles](EGameFeatureURLOptions Option, FStringView OptionName, FStringView OptionValue)
+		{
+			check(Option == EGameFeatureURLOptions::Bundles);
+			InstallBundles.Emplace(OptionValue);
+		});
+
+	if (InstallBundles.IsEmpty())
+	{
+		return true;
+	}
+
+	TSharedPtr<IInstallBundleManager> BundleManager = IInstallBundleManager::GetPlatformInstallBundleManager();
+	TValueOrError<FInstallBundleCombinedInstallState, EInstallBundleResult> InstallStateResult = BundleManager->GetInstallStateSynchronous(InstallBundles, false);
+	if (InstallStateResult.HasError())
+	{
+		UE_LOG(LogGameFeatures, Error, TEXT("Failed to get install state for PluginURL %s : Error reason %s"), *PluginURL, LexToString(InstallStateResult.GetError()));
+		return true;
+	}
+
+	return InstallStateResult.GetValue().GetAllBundlesHaveState(EInstallBundleInstallState::UpToDate);
+}
+
 bool UGameFeaturesSubsystem::IsGameFeaturePluginActive(const FString& PluginURL, bool bCheckForActivating /*= false*/) const
 {
 	if (const UGameFeaturePluginStateMachine* StateMachine = FindGameFeaturePluginStateMachine(PluginURL))
diff --git a/Engine/Plugins/Runtime/GameFeatures/Source/GameFeatures/Public/GameFeaturesSubsystem.h b/Engine/Plugins/Runtime/GameFeatures/Source/GameFeatures/Public/GameFeaturesSubsystem.h
--- a/Engine/Plugins/Runtime/GameFeatures/Source/GameFeatures/Public/GameFeaturesSubsystem.h
+++ b/Engine/Plugins/Runtime/GameFeatures/Source/GameFeatures/Public/GameFeaturesSubsystem.h
@@ -535,6 +535,9 @@ public:
 
 	/** Determines if a plugin is in the Active state.*/
 	bool IsGameFeaturePluginActive(const FString& PluginURL, bool bCheckForActivating = false) const;
+	
+	/** Determines if a plugin is up to date. Returning false means an update is available for download.*/
+	bool IsGameFeaturePluginUpToDate(const FString& PluginURL) const;
 
 	/** Deactivates the specified plugin */
 	void DeactivateGameFeaturePlugin(const FString& PluginURL);
