Replacing IsGameFeaturePluginUpToDate with DoesGameFeaturePluginNeedUpdate that checks for any bundles with the state NeedsUpdate. This is required as we only mount up to date plugins at startup and checking is up to date would return false in this scenario

diff --git a/Engine/Plugins/Runtime/GameFeatures/Source/GameFeatures/Private/GameFeaturesSubsystem.cpp b/Engine/Plugins/Runtime/GameFeatures/Source/GameFeatures/Private/GameFeaturesSubsystem.cpp
--- a/Engine/Plugins/Runtime/GameFeatures/Source/GameFeatures/Private/GameFeaturesSubsystem.cpp
+++ b/Engine/Plugins/Runtime/GameFeatures/Source/GameFeatures/Private/GameFeaturesSubsystem.cpp
@@ -1690,7 +1690,7 @@ bool UGameFeaturesSubsystem::GetGameFeaturePluginInstallPercent(TConstArrayView<
 	return false;
 }
 
-bool UGameFeaturesSubsystem::IsGameFeaturePluginUpToDate(const FString& PluginURL) const
+bool UGameFeaturesSubsystem::DoesGameFeaturePluginNeedUpdate(const FString& PluginURL) const
 {
 	TArray<FName> InstallBundles;
 	const bool bParseSuccess = UGameFeaturesSubsystem::ParsePluginURLOptions(PluginURL, EGameFeatureURLOptions::Bundles,
@@ -1702,7 +1702,7 @@ bool UGameFeaturesSubsystem::IsGameFeaturePluginUpToDate(const FString& PluginUR
 
 	if (InstallBundles.IsEmpty())
 	{
-		return true;
+		return false;
 	}
 
 	TSharedPtr<IInstallBundleManager> BundleManager = IInstallBundleManager::GetPlatformInstallBundleManager();
@@ -1710,10 +1710,10 @@ bool UGameFeaturesSubsystem::IsGameFeaturePluginUpToDate(const FString& PluginUR
 	if (InstallStateResult.HasError())
 	{
 		UE_LOG(LogGameFeatures, Error, TEXT("Failed to get install state for PluginURL %s : Error reason %s"), *PluginURL, LexToString(InstallStateResult.GetError()));
-		return true;
+		return false;
 	}
 
-	return InstallStateResult.GetValue().GetAllBundlesHaveState(EInstallBundleInstallState::UpToDate);
+	return InstallStateResult.GetValue().GetAnyBundleHasState(EInstallBundleInstallState::NeedsUpdate);
 }
 
 bool UGameFeaturesSubsystem::IsGameFeaturePluginActive(const FString& PluginURL, bool bCheckForActivating /*= false*/) const
diff --git a/Engine/Plugins/Runtime/GameFeatures/Source/GameFeatures/Public/GameFeaturesSubsystem.h b/Engine/Plugins/Runtime/GameFeatures/Source/GameFeatures/Public/GameFeaturesSubsystem.h
--- a/Engine/Plugins/Runtime/GameFeatures/Source/GameFeatures/Public/GameFeaturesSubsystem.h
+++ b/Engine/Plugins/Runtime/GameFeatures/Source/GameFeatures/Public/GameFeaturesSubsystem.h
@@ -535,9 +535,9 @@ public:
 
 	/** Determines if a plugin is in the Active state.*/
 	bool IsGameFeaturePluginActive(const FString& PluginURL, bool bCheckForActivating = false) const;
-	
-	/** Determines if a plugin is up to date. Returning false means an update is available for download.*/
-	bool IsGameFeaturePluginUpToDate(const FString& PluginURL) const;
+
+	/** Determines if a plugin is up to date or needs an update. Returns true if an update is available.*/
+	bool DoesGameFeaturePluginNeedUpdate(const FString& PluginURL) const;
 
 	/** Deactivates the specified plugin */
 	void DeactivateGameFeaturePlugin(const FString& PluginURL);
