PR #11329: Change GameFeatureSubsystem::AddObserver to notify the new observer about any already loaded plugins, addressing a TODO

diff --git a/Engine/Plugins/Runtime/GameFeatures/Source/GameFeatures/Private/GameFeaturesSubsystem.cpp b/Engine/Plugins/Runtime/GameFeatures/Source/GameFeatures/Private/GameFeaturesSubsystem.cpp
--- a/Engine/Plugins/Runtime/GameFeatures/Source/GameFeatures/Private/GameFeaturesSubsystem.cpp
+++ b/Engine/Plugins/Runtime/GameFeatures/Source/GameFeatures/Private/GameFeaturesSubsystem.cpp
@@ -637,11 +637,36 @@ void UGameFeaturesSubsystem::ForEachGameFeature(TFunctionRef<void(FGameFeatureIn
 
 void UGameFeaturesSubsystem::AddObserver(UObject* Observer)
 {
-	//@TODO: GameFeaturePluginEnginePush: May want to warn if one is added after any game feature plugins are already initialized, or go to a CallOrRegister sort of pattern
 	check(Observer);
-	if (ensureAlwaysMsgf(Cast<IGameFeatureStateChangeObserver>(Observer) != nullptr, TEXT("Observers must implement the IGameFeatureStateChangeObserver interface.")))
+	IGameFeatureStateChangeObserver* Interface = Cast<IGameFeatureStateChangeObserver>(Observer);
+	if (ensureAlwaysMsgf(Interface != nullptr, TEXT("Observers must implement the IGameFeatureStateChangeObserver interface.")))
 	{
 		Observers.AddUnique(Observer);
+
+		// Push the current state of all known game features to the new observer
+		for (auto StateMachineIt = GameFeaturePluginStateMachines.CreateConstIterator(); StateMachineIt; ++StateMachineIt)
+		{
+			if (UGameFeaturePluginStateMachine* GFSM = StateMachineIt.Value())
+			{
+				if (const UGameFeatureData* GameFeatureData = GFSM->GetGameFeatureDataForRegisteredPlugin(false))
+				{
+					FString PluginName = GFSM->GetPluginName();
+					FString PluginURL = GFSM->GetPluginURL();
+
+					Interface->OnGameFeatureRegistering(GameFeatureData, *PluginName, PluginURL);
+
+					if (GFSM->GetCurrentState() >= EGameFeaturePluginState::Loaded)
+					{
+						Interface->OnGameFeatureLoading(GameFeatureData, PluginURL);
+					}
+
+					if (GFSM->GetCurrentState() >= EGameFeaturePluginState::Active)
+					{
+						Interface->OnGameFeatureActivating(GameFeatureData, PluginURL);
+					}
+				}
+			}
+		}
 	}
 }
 
