Add support to uninstall optional content and remove from device if requested.

diff --git a/Engine/Plugins/Runtime/GameFeatures/Source/GameFeatures/Private/GameFeatureOptionalContentInstaller.cpp b/Engine/Plugins/Runtime/GameFeatures/Source/GameFeatures/Private/GameFeatureOptionalContentInstaller.cpp
--- a/Engine/Plugins/Runtime/GameFeatures/Source/GameFeatures/Private/GameFeatureOptionalContentInstaller.cpp
+++ b/Engine/Plugins/Runtime/GameFeatures/Source/GameFeatures/Private/GameFeatureOptionalContentInstaller.cpp
@@ -60,6 +60,14 @@ void UGameFeatureOptionalContentInstaller::Enable(bool bInEnable)
 	}
 }
 
+void UGameFeatureOptionalContentInstaller::UninstallContent()
+{
+	for (const FString& GFP : RelevantGFPs)
+	{
+		ReleaseContent(GFP, EInstallBundleReleaseRequestFlags::RemoveFilesIfPossible);
+	}
+}
+
 void UGameFeatureOptionalContentInstaller::EnableCellularDownloading(bool bEnable)
 {
 	if (bAllowCellDownload == bEnable)
@@ -227,7 +235,7 @@ void UGameFeatureOptionalContentInstaller::OnContentInstalled(FInstallBundleRequ
 	}
 }
 
-void UGameFeatureOptionalContentInstaller::ReleaseContent(const FString& PluginName)
+void UGameFeatureOptionalContentInstaller::ReleaseContent(const FString& PluginName, EInstallBundleReleaseRequestFlags Flags)
 {
 	TArray<FName> Bundles = GetOptionalBundlePredicate(PluginName);
 	if (Bundles.IsEmpty())
@@ -237,7 +245,7 @@ void UGameFeatureOptionalContentInstaller::ReleaseContent(const FString& PluginN
 
 	BundleManager->RequestReleaseContent(
 		Bundles, 
-		EInstallBundleReleaseRequestFlags::None, 
+		Flags,
 		{}, 
 		GameFeatureOptionalContentInstaller::InstallBundleManagerVerbosityOverride);
 }
@@ -314,4 +322,4 @@ void UGameFeatureOptionalContentInstaller::OnGameFeatureReleasing(const FString&
 	ReleaseContent(PluginName);
 
 	RelevantGFPs.Remove(PluginName);
-}
+}
\ No newline at end of file
diff --git a/Engine/Plugins/Runtime/GameFeatures/Source/GameFeatures/Public/GameFeatureOptionalContentInstaller.h b/Engine/Plugins/Runtime/GameFeatures/Source/GameFeatures/Public/GameFeatureOptionalContentInstaller.h
--- a/Engine/Plugins/Runtime/GameFeatures/Source/GameFeatures/Public/GameFeatureOptionalContentInstaller.h
+++ b/Engine/Plugins/Runtime/GameFeatures/Source/GameFeatures/Public/GameFeatureOptionalContentInstaller.h
@@ -12,6 +12,8 @@ namespace UE::GameFeatures
 	struct FResult;
 }
 
+enum class EInstallBundleReleaseRequestFlags : uint32;
+
 /** 
  * Utilty class to install GFP optional paks (usually containing optional mips) in sync with GFP content installs.
  * NOTE: This only currently supports LRU cached install bundles. It would need UI callbacks and additional support 
@@ -32,6 +34,8 @@ public:
 
 	void GAMEFEATURES_API Enable(bool bEnable);
 
+	void GAMEFEATURES_API UninstallContent();
+
 	void GAMEFEATURES_API EnableCellularDownloading(bool bEnable);
 
 private:
@@ -39,7 +43,7 @@ private:
 
 	void OnContentInstalled(FInstallBundleRequestResultInfo InResult, FString PluginName);
 
-	void ReleaseContent(const FString& PluginName);
+	void ReleaseContent(const FString& PluginName, EInstallBundleReleaseRequestFlags Flags = EInstallBundleReleaseRequestFlags::None);
 
 	void OnEnabled();
 	void OnDisabled();
