Fix issue where GFP URL being found was not returned to the caller.

diff --git a/Engine/Plugins/Runtime/GameFeatures/Source/GameFeatures/Private/GameFeaturesSubsystem.cpp b/Engine/Plugins/Runtime/GameFeatures/Source/GameFeatures/Private/GameFeaturesSubsystem.cpp
--- a/Engine/Plugins/Runtime/GameFeatures/Source/GameFeatures/Private/GameFeaturesSubsystem.cpp
+++ b/Engine/Plugins/Runtime/GameFeatures/Source/GameFeatures/Private/GameFeaturesSubsystem.cpp
@@ -2316,12 +2316,12 @@ bool UGameFeaturesSubsystem::GetBuiltInGameFeaturePluginURL(const TSharedRef<IPl
 	if (!PluginDescriptorFilename.IsEmpty() && GetDefault<UGameFeaturesSubsystemSettings>()->IsValidGameFeaturePlugin(FPaths::ConvertRelativePathToFull(PluginDescriptorFilename)) && FPaths::FileExists(PluginDescriptorFilename))
 	{
 		const FString PluginName = Plugin->GetName();
-		const bool bFoundPluginURL = GetPluginURLByName(PluginName, OutPluginURL);
+		bool bFoundPluginURL = GetPluginURLByName(PluginName, OutPluginURL);
 		if (!bFoundPluginURL)
 		{
-			GameSpecificPolicies->GetGameFeaturePluginURL(Plugin, OutPluginURL);
+			bFoundPluginURL = GameSpecificPolicies->GetGameFeaturePluginURL(Plugin, OutPluginURL);
 		}
-		return true;
+		return bFoundPluginURL;
 	}
 
 	return false;
