Add optional  FPluginDependencyDetails to ResolvePluginDependency.  FPluginDependencyDetails::bFailIfNotFound will cause GFPs load fail if set.

diff --git a/Engine/Plugins/Runtime/GameFeatures/Source/GameFeatures/Private/GameFeaturesProjectPolicies.cpp b/Engine/Plugins/Runtime/GameFeatures/Source/GameFeatures/Private/GameFeaturesProjectPolicies.cpp
--- a/Engine/Plugins/Runtime/GameFeatures/Source/GameFeatures/Private/GameFeaturesProjectPolicies.cpp
+++ b/Engine/Plugins/Runtime/GameFeatures/Source/GameFeatures/Private/GameFeaturesProjectPolicies.cpp
@@ -75,8 +75,10 @@ bool UGameFeaturesProjectPolicies::WillPluginBeCooked(const FString& PluginFilen
 	return true;
 }
 
-TValueOrError<FString, FString> UGameFeaturesProjectPolicies::ResolvePluginDependency(const FString& PluginURL, const FString& DependencyName) const
+TValueOrError<FString, FString> UGameFeaturesProjectPolicies::ResolvePluginDependency(const FString& PluginURL, const FString& DependencyName, FPluginDependencyDetails& OutDetails) const
 {
+	OutDetails = {};
+
 	FString DependencyURL;
 	bool bResolvedDependency = false;
 
diff --git a/Engine/Plugins/Runtime/GameFeatures/Source/GameFeatures/Private/GameFeaturesSubsystem.cpp b/Engine/Plugins/Runtime/GameFeatures/Source/GameFeatures/Private/GameFeaturesSubsystem.cpp
--- a/Engine/Plugins/Runtime/GameFeatures/Source/GameFeatures/Private/GameFeaturesSubsystem.cpp
+++ b/Engine/Plugins/Runtime/GameFeatures/Source/GameFeatures/Private/GameFeaturesSubsystem.cpp
@@ -3023,7 +3023,8 @@ bool UGameFeaturesSubsystem::FindOrCreatePluginDependencyStateMachines(const FSt
 		for (const FGameFeaturePluginReferenceDetails& PluginDependency : Details.PluginDependencies)
 		{
 			const FString& DependencyName = PluginDependency.PluginName;
-			TValueOrError<FString, FString> DependencyURLInfo = GameSpecificPolicies->ResolvePluginDependency(PluginURL, DependencyName);
+			FPluginDependencyDetails DependencyDetails;
+			TValueOrError<FString, FString> DependencyURLInfo = GameSpecificPolicies->ResolvePluginDependency(PluginURL, DependencyName, DependencyDetails);
 			if (DependencyURLInfo.HasError())
 			{
 				FString ErrorMessage = FString::Printf(TEXT("Game feature plugin '%s' has unknown dependency '%s' [%s]."), *PluginFilename, *DependencyName, *DependencyURLInfo.GetError());
@@ -3031,10 +3032,18 @@ bool UGameFeaturesSubsystem::FindOrCreatePluginDependencyStateMachines(const FSt
 
 				SetExplanationForNotMountingPlugin(PluginURL, ErrorMessage);
 
-				//Don't actually return false here as we want to still be able to progress in the case of 
-				//things like an editor plugin being included as a dependency in the client or a dynamic dependency that
-				//hasn't correctly loaded yet
-				continue;
+				if (DependencyDetails.bFailIfNotFound)
+				{
+					// This plugin is known to be a hard dependency so fail early before something harder to debug happens.
+					return false;
+				}
+				else
+				{
+					// Don't actually return false here as we want to still be able to progress in the case of 
+					// things like an editor plugin being included as a dependency in the client. We can't tell from just 
+					// the reference if a plugin is not enabled for this build target.
+					continue;
+				}
 			}
 
 			const FString& DependencyURL = DependencyURLInfo.GetValue();
diff --git a/Engine/Plugins/Runtime/GameFeatures/Source/GameFeatures/Public/GameFeaturesProjectPolicies.h b/Engine/Plugins/Runtime/GameFeatures/Source/GameFeatures/Public/GameFeaturesProjectPolicies.h
--- a/Engine/Plugins/Runtime/GameFeatures/Source/GameFeatures/Public/GameFeaturesProjectPolicies.h
+++ b/Engine/Plugins/Runtime/GameFeatures/Source/GameFeatures/Public/GameFeaturesProjectPolicies.h
@@ -9,6 +9,11 @@
 
 class UGameFeatureData;
 
+struct FPluginDependencyDetails
+{
+	bool bFailIfNotFound = false;
+};
+
 // This class allows project-specific rules to be implemented for game feature plugins.
 // Create a subclass and choose it in Project Settings .. Game Features
 UCLASS()
@@ -47,7 +52,12 @@ public:
 	// Called to resolve plugin dependencies, will successfully return an empty string if a dependency is not a GFP.
 	// This may be called with file protocol for built-in plugins in some cases, even if a different protocol is used at runtime.
 	// returns The dependency URL or an error if the dependency could not be resolved
-	virtual TValueOrError<FString, FString> ResolvePluginDependency(const FString& PluginURL, const FString& DependencyName) const;
+	virtual TValueOrError<FString, FString> ResolvePluginDependency(const FString& PluginURL, const FString& DependencyName, FPluginDependencyDetails& OutDetails) const;
+	virtual TValueOrError<FString, FString> ResolvePluginDependency(const FString& PluginURL, const FString& DependencyName) const
+	{
+		FPluginDependencyDetails Details;
+		return ResolvePluginDependency(PluginURL, DependencyName, Details);
+	}
 
 	// Called to resolve install bundles for streaming asset dependencies
 	virtual TValueOrError<TArray<FName>, FString> GetStreamingAssetInstallBundles(FStringView PluginURL) const { return MakeValue(); }
