Fix wrong overload of ResolvePluginDependency having default implementation.

diff --git a/Engine/Plugins/Runtime/GameFeatures/Source/GameFeatures/Private/GameFeaturesProjectPolicies.cpp b/Engine/Plugins/Runtime/GameFeatures/Source/GameFeatures/Private/GameFeaturesProjectPolicies.cpp
--- a/Engine/Plugins/Runtime/GameFeatures/Source/GameFeatures/Private/GameFeaturesProjectPolicies.cpp
+++ b/Engine/Plugins/Runtime/GameFeatures/Source/GameFeatures/Private/GameFeaturesProjectPolicies.cpp
@@ -78,7 +78,11 @@ bool UGameFeaturesProjectPolicies::WillPluginBeCooked(const FString& PluginFilen
 TValueOrError<FString, FString> UGameFeaturesProjectPolicies::ResolvePluginDependency(const FString& PluginURL, const FString& DependencyName, FPluginDependencyDetails& OutDetails) const
 {
 	OutDetails = {};
+	return ResolvePluginDependency(PluginURL, DependencyName);
+}
 
+TValueOrError<FString, FString> UGameFeaturesProjectPolicies::ResolvePluginDependency(const FString& PluginURL, const FString& DependencyName) const
+{
 	FString DependencyURL;
 	bool bResolvedDependency = false;
 
diff --git a/Engine/Plugins/Runtime/GameFeatures/Source/GameFeatures/Public/GameFeaturesProjectPolicies.h b/Engine/Plugins/Runtime/GameFeatures/Source/GameFeatures/Public/GameFeaturesProjectPolicies.h
--- a/Engine/Plugins/Runtime/GameFeatures/Source/GameFeatures/Public/GameFeaturesProjectPolicies.h
+++ b/Engine/Plugins/Runtime/GameFeatures/Source/GameFeatures/Public/GameFeaturesProjectPolicies.h
@@ -53,11 +53,7 @@ public:
 	// This may be called with file protocol for built-in plugins in some cases, even if a different protocol is used at runtime.
 	// returns The dependency URL or an error if the dependency could not be resolved
 	virtual TValueOrError<FString, FString> ResolvePluginDependency(const FString& PluginURL, const FString& DependencyName, FPluginDependencyDetails& OutDetails) const;
-	virtual TValueOrError<FString, FString> ResolvePluginDependency(const FString& PluginURL, const FString& DependencyName) const
-	{
-		FPluginDependencyDetails Details;
-		return ResolvePluginDependency(PluginURL, DependencyName, Details);
-	}
+	virtual TValueOrError<FString, FString> ResolvePluginDependency(const FString& PluginURL, const FString& DependencyName) const;
 
 	// Called to resolve install bundles for streaming asset dependencies
 	virtual TValueOrError<TArray<FName>, FString> GetStreamingAssetInstallBundles(FStringView PluginURL) const { return MakeValue(); }
