On second boot attempt to request optional textures for plugins which might not download

diff --git a/Engine/Plugins/Runtime/GameFeatures/Source/GameFeatures/Private/GameFeatureOptionalContentInstaller.cpp b/Engine/Plugins/Runtime/GameFeatures/Source/GameFeatures/Private/GameFeatureOptionalContentInstaller.cpp
--- a/Engine/Plugins/Runtime/GameFeatures/Source/GameFeatures/Private/GameFeatureOptionalContentInstaller.cpp
+++ b/Engine/Plugins/Runtime/GameFeatures/Source/GameFeatures/Private/GameFeatureOptionalContentInstaller.cpp
@@ -257,8 +257,7 @@ void UGameFeatureOptionalContentInstaller::OnEnabled()
 
 	UGameFeaturesSubsystem::Get().ForEachGameFeature([this](FGameFeatureInfo&& Info) -> void
 	{
-		if (UGameFeaturesSubsystem::GetPluginURLProtocol(Info.URL) == EGameFeaturePluginProtocol::InstallBundle &&
-			Info.CurrentState >= EGameFeaturePluginState::Downloading)
+		if (Info.CurrentState >= EGameFeaturePluginState::Downloading)
 		{
 			if (UpdateContent(Info.Name, false))
 			{
@@ -312,6 +311,20 @@ void UGameFeatureOptionalContentInstaller::OnGameFeatureDownloading(const FStrin
 	}
 }
 
+void UGameFeatureOptionalContentInstaller::OnGameFeatureRegistering(const UGameFeatureData* GameFeatureData, const FString& PluginName, const FString& PluginURL)
+{
+    // Used for already downloaded cached plugins that do not download at startup but register.
+    if (!IsEnabled() || RelevantGFPs.Contains(PluginName))
+    {
+        return;
+    }
+
+    if (UpdateContent(PluginName, false))
+    {
+        RelevantGFPs.Add(PluginName);
+    }
+}
+
 void UGameFeatureOptionalContentInstaller::OnGameFeatureReleasing(const FString& PluginName, const FGameFeaturePluginIdentifier& PluginIdentifier)
 {
 	if (!IsEnabled())
@@ -322,4 +335,4 @@ void UGameFeatureOptionalContentInstaller::OnGameFeatureReleasing(const FString&
 	ReleaseContent(PluginName);
 
 	RelevantGFPs.Remove(PluginName);
-}
\ No newline at end of file
+}
diff --git a/Engine/Plugins/Runtime/GameFeatures/Source/GameFeatures/Public/GameFeatureOptionalContentInstaller.h b/Engine/Plugins/Runtime/GameFeatures/Source/GameFeatures/Public/GameFeatureOptionalContentInstaller.h
--- a/Engine/Plugins/Runtime/GameFeatures/Source/GameFeatures/Public/GameFeatureOptionalContentInstaller.h
+++ b/Engine/Plugins/Runtime/GameFeatures/Source/GameFeatures/Public/GameFeatureOptionalContentInstaller.h
@@ -55,6 +55,7 @@ private:
 	// IGameFeatureStateChangeObserver Interface
 	virtual void OnGameFeaturePredownloading(const FString& PluginName, const FGameFeaturePluginIdentifier& PluginIdentifier) override;
 	virtual void OnGameFeatureDownloading(const FString& PluginName, const FGameFeaturePluginIdentifier& PluginIdentifier) override;
+    virtual void OnGameFeatureRegistering(const UGameFeatureData* GameFeatureData, const FString& PluginName, const FString& PluginURL) override;
 	virtual void OnGameFeatureReleasing(const FString& PluginName, const FGameFeaturePluginIdentifier& PluginIdentifier) override;
 
 private:
