Adding logging for Optional Bundle install and uninstall

diff --git a/Engine/Plugins/Runtime/GameFeatures/Source/GameFeatures/Private/GameFeatureOptionalContentInstaller.cpp b/Engine/Plugins/Runtime/GameFeatures/Source/GameFeatures/Private/GameFeatureOptionalContentInstaller.cpp
--- a/Engine/Plugins/Runtime/GameFeatures/Source/GameFeatures/Private/GameFeatureOptionalContentInstaller.cpp
+++ b/Engine/Plugins/Runtime/GameFeatures/Source/GameFeatures/Private/GameFeatureOptionalContentInstaller.cpp
@@ -64,6 +64,7 @@ void UGameFeatureOptionalContentInstaller::UninstallContent()
 {
 	for (const FString& GFP : RelevantGFPs)
 	{
+		UE_LOG(LogGameFeatures, Log, TEXT("Uninstalling Optional bundles for %s"), *GFP);
 		ReleaseContent(GFP, EInstallBundleReleaseRequestFlags::RemoveFilesIfPossible);
 	}
 }
@@ -106,6 +107,10 @@ bool UGameFeatureOptionalContentInstaller::UpdateContent(const FString& PluginNa
 		return false;
 	}
 
+	for (const FName& Bundle : Bundles)
+	{
+		UE_LOG(LogGameFeatures, Log, TEXT("Requesting update for %s"), *Bundle.ToString());
+	}
 	EInstallBundleRequestFlags InstallFlags = EInstallBundleRequestFlags::AsyncMount;
 	if (bIsPredownload)
 	{
@@ -170,6 +175,7 @@ void UGameFeatureOptionalContentInstaller::OnContentInstalled(FInstallBundleRequ
 
 	GFPInstall.BundlesEnqueued.Remove(InResult.BundleName);
 
+	UE_LOG(LogGameFeatures, Log, TEXT("Finished install for %s"), *InResult.BundleName.ToString());
 	if (InResult.Result != EInstallBundleResult::OK)
 	{
 		if (InResult.OptionalErrorCode.IsEmpty())
