#UE Implement == operator for FGameFeatureProtocolOptions. It did not have the default == operator because TUnion implements one, but TUnion does not consider the properties in FGameFeatureProtocolOptions

diff --git a/Engine/Plugins/Runtime/GameFeatures/Source/GameFeatures/Public/GameFeaturesSubsystem.h b/Engine/Plugins/Runtime/GameFeatures/Source/GameFeatures/Public/GameFeaturesSubsystem.h
--- a/Engine/Plugins/Runtime/GameFeatures/Source/GameFeatures/Public/GameFeaturesSubsystem.h
+++ b/Engine/Plugins/Runtime/GameFeatures/Source/GameFeatures/Public/GameFeaturesSubsystem.h
@@ -368,6 +368,14 @@ struct FGameFeatureProtocolOptions : public TUnion<FInstallBundlePluginProtocolO
 	GAMEFEATURES_API explicit FGameFeatureProtocolOptions(const FInstallBundlePluginProtocolOptions& InOptions);
 	GAMEFEATURES_API explicit FGameFeatureProtocolOptions(FNull InOptions);
 
+	bool operator==(const FGameFeatureProtocolOptions& Other) const
+	{
+		return TUnion<FInstallBundlePluginProtocolOptions, FNull>::operator==(Other) &&
+			bForceSyncLoading == Other.bForceSyncLoading &&
+			bLogWarningOnForcedDependencyCreation == Other.bLogWarningOnForcedDependencyCreation &&
+			bLogErrorOnForcedDependencyCreation == Other.bLogErrorOnForcedDependencyCreation;
+	}
+
 	/** Force this GFP to load synchronously even if async loading is allowed */
 	bool bForceSyncLoading : 1;
 
