Add a GameFeaturePlugin Activated callback to trigger only once a plugin is completely active; Have ConversationRegistry rely on this new callback to ensure no active files are accidentally skipped when rebuilding the registry.

diff --git a/Engine/Plugins/Experimental/CommonConversation/Source/CommonConversationRuntime/Private/ConversationRegistry.cpp b/Engine/Plugins/Experimental/CommonConversation/Source/CommonConversationRuntime/Private/ConversationRegistry.cpp
--- a/Engine/Plugins/Experimental/CommonConversation/Source/CommonConversationRuntime/Private/ConversationRegistry.cpp
+++ b/Engine/Plugins/Experimental/CommonConversation/Source/CommonConversationRuntime/Private/ConversationRegistry.cpp
@@ -185,7 +185,7 @@ void UConversationRegistry::GameFeatureStateModified()
 	}
 }
 
-void UConversationRegistry::OnGameFeatureActivating(const UGameFeatureData* GameFeatureData, const FString& PluginURL)
+void UConversationRegistry::OnGameFeatureActivated(const UGameFeatureData* GameFeatureData, const FString& PluginURL)
 {
 	GameFeatureStateModified();
 }
diff --git a/Engine/Plugins/Experimental/CommonConversation/Source/CommonConversationRuntime/Public/ConversationRegistry.h b/Engine/Plugins/Experimental/CommonConversation/Source/CommonConversationRuntime/Public/ConversationRegistry.h
--- a/Engine/Plugins/Experimental/CommonConversation/Source/CommonConversationRuntime/Public/ConversationRegistry.h
+++ b/Engine/Plugins/Experimental/CommonConversation/Source/CommonConversationRuntime/Public/ConversationRegistry.h
@@ -113,7 +113,7 @@ private:
 
 	void GameFeatureStateModified();
 
-	virtual void OnGameFeatureActivating(const UGameFeatureData* GameFeatureData, const FString& PluginURL) override;
+	virtual void OnGameFeatureActivated(const UGameFeatureData* GameFeatureData, const FString& PluginURL) override;
 
 	virtual void OnGameFeatureDeactivating(const UGameFeatureData* GameFeatureData, FGameFeatureDeactivatingContext& Context, const FString& PluginURL) override;
 
diff --git a/Engine/Plugins/Runtime/GameFeatures/Source/GameFeatures/Private/GameFeaturePluginStateMachine.cpp b/Engine/Plugins/Runtime/GameFeatures/Source/GameFeatures/Private/GameFeaturePluginStateMachine.cpp
--- a/Engine/Plugins/Runtime/GameFeatures/Source/GameFeatures/Private/GameFeaturePluginStateMachine.cpp
+++ b/Engine/Plugins/Runtime/GameFeatures/Source/GameFeatures/Private/GameFeaturePluginStateMachine.cpp
@@ -3182,6 +3182,17 @@ struct FGameFeaturePluginState_Active : public FDestinationGameFeaturePluginStat
 {
 	FGameFeaturePluginState_Active(FGameFeaturePluginStateMachineProperties& InStateProperties) : FDestinationGameFeaturePluginState(InStateProperties) {}
 
+	virtual void BeginState() override
+	{
+		TRACE_CPUPROFILER_EVENT_SCOPE(GFP_Active);
+		check(GEngine);
+
+		{
+			TRACE_CPUPROFILER_EVENT_SCOPE(GFP_Active_SendEvents);
+			UGameFeaturesSubsystem::Get().OnGameFeatureActivated(StateProperties.GameFeatureData, StateProperties.PluginName, StateProperties.PluginIdentifier);
+		}
+	}
+
 	virtual void UpdateState(FGameFeaturePluginStateStatus& StateStatus) override
 	{
 		if (StateProperties.Destination < EGameFeaturePluginState::Active)
diff --git a/Engine/Plugins/Runtime/GameFeatures/Source/GameFeatures/Private/GameFeaturesSubsystem.cpp b/Engine/Plugins/Runtime/GameFeatures/Source/GameFeatures/Private/GameFeaturesSubsystem.cpp
--- a/Engine/Plugins/Runtime/GameFeatures/Source/GameFeatures/Private/GameFeaturesSubsystem.cpp
+++ b/Engine/Plugins/Runtime/GameFeatures/Source/GameFeatures/Private/GameFeaturesSubsystem.cpp
@@ -663,6 +663,7 @@ void UGameFeaturesSubsystem::AddObserver(UObject* Observer)
 					if (GFSM->GetCurrentState() >= EGameFeaturePluginState::Active)
 					{
 						Interface->OnGameFeatureActivating(GameFeatureData, PluginURL);
+						Interface->OnGameFeatureActivated(GameFeatureData, PluginURL);
 					}
 				}
 			}
@@ -1123,6 +1124,25 @@ void UGameFeaturesSubsystem::OnGameFeatureActivating(const UGameFeatureData* Gam
 	}
 }
 
+void UGameFeaturesSubsystem::OnGameFeatureActivated(const UGameFeatureData* GameFeatureData, const FString& PluginName, const FGameFeaturePluginIdentifier& PluginIdentifier)
+{
+	{
+		TRACE_CPUPROFILER_EVENT_SCOPE(GFP_OnActivated_CallbackObservers);
+		CallbackObservers(EObserverCallback::Activated, PluginIdentifier, &PluginName, GameFeatureData);
+	}
+
+	{
+		TRACE_CPUPROFILER_EVENT_SCOPE(GFP_OnActivated_CallbackActions);
+		for (UGameFeatureAction* Action : GameFeatureData->GetActions())
+		{
+			if (Action != nullptr)
+			{
+				Action->OnGameFeatureActivated();
+			}
+		}
+	}
+}
+
 void UGameFeaturesSubsystem::OnGameFeatureDeactivating(const UGameFeatureData* GameFeatureData, const FString& PluginName, FGameFeatureDeactivatingContext& Context, const FGameFeaturePluginIdentifier& PluginIdentifier)
 {
 #if !WITH_EDITOR
@@ -3251,7 +3271,7 @@ void UGameFeaturesSubsystem::CallbackObservers(EObserverCallback CallbackType, c
 	const UGameFeatureData* GameFeatureData /*= nullptr*/, 
 	FGameFeatureStateChangeContext* StateChangeContext /*= nullptr*/)
 {
-	static_assert(std::underlying_type<EObserverCallback>::type(EObserverCallback::Count) == 14, "Update UGameFeaturesSubsystem::CallbackObservers to handle added EObserverCallback");
+	static_assert(std::underlying_type<EObserverCallback>::type(EObserverCallback::Count) == 15, "Update UGameFeaturesSubsystem::CallbackObservers to handle added EObserverCallback");
 
 	// Protect against modifying the observer list during iteration
 	TArray<UObject*> LocalObservers(Observers);
@@ -3376,6 +3396,15 @@ void UGameFeaturesSubsystem::CallbackObservers(EObserverCallback CallbackType, c
 		}
 		break;
 	}
+	case EObserverCallback::Activated:
+	{
+		check(GameFeatureData);
+		for (UObject* Observer : LocalObservers)
+		{
+			CastChecked<IGameFeatureStateChangeObserver>(Observer)->OnGameFeatureActivated(GameFeatureData, PluginIdentifier.GetFullPluginURL());
+		}
+		break;
+	}
 	case EObserverCallback::Deactivating:
 	{
 #if !WITH_EDITOR
diff --git a/Engine/Plugins/Runtime/GameFeatures/Source/GameFeatures/Public/GameFeatureAction.h b/Engine/Plugins/Runtime/GameFeatures/Source/GameFeatures/Public/GameFeatureAction.h
--- a/Engine/Plugins/Runtime/GameFeatures/Source/GameFeatures/Public/GameFeatureAction.h
+++ b/Engine/Plugins/Runtime/GameFeatures/Source/GameFeatures/Public/GameFeatureAction.h
@@ -36,6 +36,9 @@ public:
 	/** Older-style activation function with no context, called by base class if context version is not overridden */
 	virtual void OnGameFeatureActivating() {}
 
+	/** Called when the feature is fully active */
+	virtual void OnGameFeatureActivated() {}
+
 	/** Called when game feature is deactivated, it may be activated again in the near future */
 	virtual void OnGameFeatureDeactivating(FGameFeatureDeactivatingContext& Context) {}
 
diff --git a/Engine/Plugins/Runtime/GameFeatures/Source/GameFeatures/Public/GameFeatureStateChangeObserver.h b/Engine/Plugins/Runtime/GameFeatures/Source/GameFeatures/Public/GameFeatureStateChangeObserver.h
--- a/Engine/Plugins/Runtime/GameFeatures/Source/GameFeatures/Public/GameFeatureStateChangeObserver.h
+++ b/Engine/Plugins/Runtime/GameFeatures/Source/GameFeatures/Public/GameFeatureStateChangeObserver.h
@@ -70,6 +70,9 @@ public:
 	// Invoked prior to activating a game feature plugin
 	virtual void OnGameFeatureActivating(const UGameFeatureData* GameFeatureData, const FString& PluginURL) {}
 
+	// Invoked after a game feature plugin is activated
+	virtual void OnGameFeatureActivated(const UGameFeatureData* GameFeatureData, const FString& PluginURL) {}
+
 	// Invoked prior to deactivating a game feature plugin
 	virtual void OnGameFeatureDeactivating(const UGameFeatureData* GameFeatureData, FGameFeatureDeactivatingContext& Context, const FString& PluginURL) {}
 
diff --git a/Engine/Plugins/Runtime/GameFeatures/Source/GameFeatures/Public/GameFeaturesSubsystem.h b/Engine/Plugins/Runtime/GameFeatures/Source/GameFeatures/Public/GameFeaturesSubsystem.h
--- a/Engine/Plugins/Runtime/GameFeatures/Source/GameFeatures/Public/GameFeaturesSubsystem.h
+++ b/Engine/Plugins/Runtime/GameFeatures/Source/GameFeatures/Public/GameFeaturesSubsystem.h
@@ -688,6 +688,9 @@ private:
 	void OnGameFeatureActivating(const UGameFeatureData* GameFeatureData, const FString& PluginName, FGameFeatureActivatingContext& Context, const FGameFeaturePluginIdentifier& PluginIdentifier);
 	friend struct FGameFeaturePluginState_Activating;
 
+	void OnGameFeatureActivated(const UGameFeatureData* GameFeatureData, const FString& PluginName, const FGameFeaturePluginIdentifier& PluginIdentifier);
+	friend struct FGameFeaturePluginState_Active;
+
 	void OnGameFeatureDeactivating(const UGameFeatureData* GameFeatureData, const FString& PluginName, FGameFeatureDeactivatingContext& Context, const FGameFeaturePluginIdentifier& PluginIdentifier);
 	friend struct FGameFeaturePluginState_Deactivating;
 
@@ -803,6 +806,7 @@ private:
 		Loading,
 		Unloading,
 		Activating,
+		Activated,
 		Deactivating,
 		PauseChanged,
 		Count
