Improved error logging for gamefeature subsystem.

diff --git a/Engine/Plugins/Runtime/GameFeatures/Source/GameFeatures/Private/GameFeaturesSubsystem.cpp b/Engine/Plugins/Runtime/GameFeatures/Source/GameFeatures/Private/GameFeaturesSubsystem.cpp
--- a/Engine/Plugins/Runtime/GameFeatures/Source/GameFeatures/Private/GameFeaturesSubsystem.cpp
+++ b/Engine/Plugins/Runtime/GameFeatures/Source/GameFeatures/Private/GameFeaturesSubsystem.cpp
@@ -497,10 +497,10 @@ void UGameFeaturesSubsystem::OnAssetManagerCreated()
 	GameSpecificPolicies->InitGameFeatureManager();
 }
 
-bool UGameFeaturesSubsystem::IsPluginAllowed(const FString& PluginURL) const
+bool UGameFeaturesSubsystem::IsPluginAllowed(const FString& PluginURL, FString* OutReason) const
 {
 	ensureMsgf(bInitializedPolicyManager, TEXT("Attemting to load plugin [%s] before GameFeaturesSubsystem is ready!"), *PluginURL);
-	return bInitializedPolicyManager && GameSpecificPolicies->IsPluginAllowed(PluginURL);
+	return bInitializedPolicyManager && GameSpecificPolicies->IsPluginAllowed(PluginURL, OutReason);
 }
 
 TSharedPtr<FStreamableHandle> UGameFeaturesSubsystem::LoadGameFeatureData(const FString& GameFeatureToLoad, bool bStartStalled /*= false*/)
@@ -1269,10 +1269,12 @@ void UGameFeaturesSubsystem::LoadGameFeaturePlugin(const FString& PluginURL, con
 
 void UGameFeaturesSubsystem::LoadGameFeaturePlugin(const FString& PluginURL, const FGameFeatureProtocolOptions& ProtocolOptions, const FGameFeaturePluginLoadComplete& CompleteDelegate)
 {
-	const bool bIsPluginAllowed = IsPluginAllowed(PluginURL);
+	FString Reason;
+	const bool bIsPluginAllowed = IsPluginAllowed(PluginURL, &Reason);
 	if (!bIsPluginAllowed)
 	{
-		CompleteDelegate.ExecuteIfBound(UE::GameFeatures::FResult(MakeError(UE::GameFeatures::SubsystemErrorNamespace + UE::GameFeatures::CommonErrorCodes::PluginNotAllowed)));
+		CompleteDelegate.ExecuteIfBound(UE::GameFeatures::FResult(MakeError(UE::GameFeatures::SubsystemErrorNamespace + UE::GameFeatures::CommonErrorCodes::PluginNotAllowed), 
+			FText::Format(NSLOCTEXT("GameFeatures", "PluginDisallowedText","Plugin disallowed reason: {0}"), FText::FromString(Reason)) ));
 		return;
 	}
 
diff --git a/Engine/Plugins/Runtime/GameFeatures/Source/GameFeatures/Public/GameFeaturesProjectPolicies.h b/Engine/Plugins/Runtime/GameFeatures/Source/GameFeatures/Public/GameFeaturesProjectPolicies.h
--- a/Engine/Plugins/Runtime/GameFeatures/Source/GameFeatures/Public/GameFeaturesProjectPolicies.h
+++ b/Engine/Plugins/Runtime/GameFeatures/Source/GameFeatures/Public/GameFeaturesProjectPolicies.h
@@ -47,7 +47,7 @@ public:
 
 	// Called to determine if a plugin is allowed to be loaded or not
 	// (e.g., when doing a fast cook a game might want to disable some or all game feature plugins)
-	virtual bool IsPluginAllowed(const FString& PluginURL) const { return true; }
+	virtual bool IsPluginAllowed(const FString& PluginURL, FString* OutReseason = nullptr) const { return true; }
 
 	// Called to resolve plugin dependencies, will successfully return an empty string if a dependency is not a GFP.
 	// This may be called with file protocol for built-in plugins in some cases, even if a different protocol is used at runtime.
diff --git a/Engine/Plugins/Runtime/GameFeatures/Source/GameFeatures/Public/GameFeaturesSubsystem.h b/Engine/Plugins/Runtime/GameFeatures/Source/GameFeatures/Public/GameFeaturesSubsystem.h
--- a/Engine/Plugins/Runtime/GameFeatures/Source/GameFeatures/Public/GameFeaturesSubsystem.h
+++ b/Engine/Plugins/Runtime/GameFeatures/Source/GameFeatures/Public/GameFeaturesSubsystem.h
@@ -715,7 +715,7 @@ private:
 	void GetExplanationForUnavailablePackage(const FString& SkippedPackage, IPlugin* PluginIfFound, FStringBuilderBase& InOutExplanation);
 
 private:
-	bool IsPluginAllowed(const FString& PluginURL) const;
+	bool IsPluginAllowed(const FString& PluginURL, FString* OutReason = nullptr) const;
 
 	bool ShouldUpdatePluginProtocolOptions(const UGameFeaturePluginStateMachine* StateMachine, const FGameFeatureProtocolOptions& NewOptions);
 	UE::GameFeatures::FResult UpdateGameFeatureProtocolOptions(UGameFeaturePluginStateMachine* StateMachine, const FGameFeatureProtocolOptions& NewOptions, bool* bOutDidUpdate = nullptr);
