General IAD Support - GFP subsystem Changes
- Deprecate GetStreamingAssetInstallBundles and replace with GetStreamingAssetInstallModes which returns whether or not to full install an IAD GFP dependency.
- GFP state machine is now responsible for installing packages from IAD
- Refactor GFP shared install bundle dependency tracking
-- Tracks IAD handles for each bundle now
-- Install bundle dependencies are now completely expanded when checking status and properly tracked
-- Deprecate FInstallBundlePluginProtocolOptions::ReleaseInstallBundleFlags since the dependency tracking makes it possibe to correctly set all flags internally

diff --git a/Engine/Plugins/Runtime/GameFeatures/Source/GameFeatures/GameFeatures.Build.cs b/Engine/Plugins/Runtime/GameFeatures/Source/GameFeatures/GameFeatures.Build.cs
--- a/Engine/Plugins/Runtime/GameFeatures/Source/GameFeatures/GameFeatures.Build.cs
+++ b/Engine/Plugins/Runtime/GameFeatures/Source/GameFeatures/GameFeatures.Build.cs
@@ -24,6 +24,7 @@ namespace UnrealBuildTool.Rules
 					"AssetRegistry",
 					"GameplayTags",
 					"InstallBundleManager",
+					"IoStoreOnDemand",
 					"Json",
 					"JsonUtilities",
 					"PakFile",
diff --git a/Engine/Plugins/Runtime/GameFeatures/Source/GameFeatures/Private/GameFeaturePluginStateMachine.cpp b/Engine/Plugins/Runtime/GameFeatures/Source/GameFeatures/Private/GameFeaturePluginStateMachine.cpp
--- a/Engine/Plugins/Runtime/GameFeatures/Source/GameFeatures/Private/GameFeaturePluginStateMachine.cpp
+++ b/Engine/Plugins/Runtime/GameFeatures/Source/GameFeatures/Private/GameFeaturePluginStateMachine.cpp
@@ -12,6 +12,7 @@
 #include "BundlePrereqCombinedStatusHelper.h"
 #include "Interfaces/IPluginManager.h"
 #include "Internationalization/PackageLocalizationManager.h"
+#include "IO/IoStoreOnDemand.h"
 #include "Logging/StructuredLog.h"
 #include "Materials/MaterialInterface.h"
 #include "Misc/AsciiSet.h"
@@ -23,7 +24,9 @@
 #include "Misc/PackageName.h"
 #include "Misc/ScopedSlowTask.h"
 #include "Misc/WildcardString.h"
+#include "Algo/Accumulate.h"
 #include "Algo/AllOf.h"
+#include "Algo/Find.h"
 #include "Misc/TVariantMeta.h"
 #include "RenderDeferredCleanup.h"
 #include "Serialization/MemoryReader.h"
@@ -288,6 +291,149 @@ namespace UE::GameFeatures
 		}
 	}
 #endif //if WITH_EDITOR
+
+	enum class EGFPInstallLevel : uint8
+	{
+		Download = 0,
+		Mount,
+		AssetStream
+	};
+
+	struct FPluginInstallBundleReferencers
+	{
+		// GFPs using an install bundle and the relevant state of that GFP
+		TMap<FStringView, EGFPInstallLevel> GFPs;
+	};
+
+	struct FPluginIoStoreOnDemandHandles
+	{
+		// IoStoreOnDemand assets required for initial download
+		UE::IoStore::FOnDemandContentHandle DownloadHandle;
+		// IoStoreOnDemand assets from AssetDependencyStreaming
+		UE::IoStore::FOnDemandContentHandle StreamInHandle;
+	};
+
+	// class to manage GFPs sharing installation data
+	class FGFPSharedInstallTracker
+	{
+	public:
+		// The caller should pass all resolved bundle depenencies 
+		void AddBundleRefs(const FStringView PluginName, const EGFPInstallLevel Level, const TConstArrayView<FName> Bundles)
+		{
+			for (const FName BundleName : Bundles)
+			{
+				FPluginInstallBundleReferencers& PluginRefs = InstallBundleToGFPRefs.FindOrAdd(BundleName);
+				PluginRefs.GFPs.Emplace(PluginName, Level);
+			}
+		}
+
+		// The caller should pass all resolved bundle depenencies 
+		UE::IoStore::FOnDemandContentHandle AddOnDemandContentHandle(const FName Bundle, const EGFPInstallLevel Level)
+		{
+			switch(Level)
+			{
+			case EGFPInstallLevel::Download:
+			{
+				FPluginIoStoreOnDemandHandles& Handles = InstallBundleToOnDemandHandles.FindOrAdd(Bundle);
+				if (!Handles.DownloadHandle.IsValid())
+				{
+					Handles.DownloadHandle = UE::IoStore::FOnDemandContentHandle::Create(Bundle.ToString());
+				}
+				return Handles.DownloadHandle;
+			}
+			case EGFPInstallLevel::AssetStream:
+			{
+				FPluginIoStoreOnDemandHandles& Handles = InstallBundleToOnDemandHandles.FindOrAdd(Bundle);
+				if (!Handles.StreamInHandle.IsValid())
+				{
+					Handles.StreamInHandle = UE::IoStore::FOnDemandContentHandle::Create(Bundle.ToString() + TEXTVIEW("/deps"));
+				}
+				return Handles.StreamInHandle;
+			}
+			};
+
+			return {};
+		}
+
+		// The caller should pass all resolved bundle depenencies 
+		TArray<FName> Release(
+			const FStringView PluginName, const EGFPInstallLevel InLevel, const TConstArrayView<FName> Bundles)
+		{
+			TArray<FName> BundlesToRelease;
+			for (const FName Bundle : Bundles)
+			{
+				bool bRelease = true;
+
+				FPluginInstallBundleReferencers* PluginRefs = InstallBundleToGFPRefs.Find(Bundle);
+				if (PluginRefs)
+				{
+					EGFPInstallLevel* Level = PluginRefs->GFPs.FindByHash(GetTypeHash(PluginName), PluginName);
+					if (Level && *Level >= InLevel)
+					{
+						switch(*Level)
+						{
+						case EGFPInstallLevel::Download:
+							PluginRefs->GFPs.RemoveByHash(GetTypeHash(PluginName), PluginName);
+							break;
+						case EGFPInstallLevel::Mount:
+							(*Level) = EGFPInstallLevel::Download;
+							break;
+						case EGFPInstallLevel::AssetStream:
+							(*Level) = EGFPInstallLevel::Mount;
+							break;
+						}
+					}
+
+					for (TPair<FStringView, EGFPInstallLevel>& GFPRef : PluginRefs->GFPs)
+					{
+						if (GFPRef.Value >= InLevel)
+						{
+							bRelease = false;
+							break;
+						}
+					}
+
+					if (PluginRefs->GFPs.IsEmpty())
+					{
+						InstallBundleToGFPRefs.Remove(Bundle);
+					}
+				}
+
+				if (bRelease)
+				{
+					BundlesToRelease.Add(Bundle);
+				}
+			}
+
+			switch(InLevel)
+			{
+			case EGFPInstallLevel::Download:
+				for (const FName Bundle : BundlesToRelease)
+				{
+					InstallBundleToOnDemandHandles.Remove(Bundle);
+				}
+				break;
+			case EGFPInstallLevel::AssetStream:
+				for (const FName& Bundle : BundlesToRelease)
+				{
+					FPluginIoStoreOnDemandHandles* Handles = InstallBundleToOnDemandHandles.Find(Bundle);
+					if (Handles)
+					{
+						Handles->StreamInHandle.Reset();
+					}
+				}
+				break;
+			}
+
+			return BundlesToRelease;
+		}
+
+	private:
+		TMap<FName, FPluginInstallBundleReferencers> InstallBundleToGFPRefs;
+		TMap<FName, FPluginIoStoreOnDemandHandles> InstallBundleToOnDemandHandles;
+	};
+
+	static FGFPSharedInstallTracker GFPSharedInstallTracker;
 }
 
 void FGameFeaturePluginStateStatus::SetTransition(EGameFeaturePluginState InTransitionToState)
@@ -810,9 +956,11 @@ struct FGameFeaturePluginState_CheckingStatus : public FGameFeaturePluginState
 				return;
 			}
 
-			const TArray<FName>& InstallBundles = StateProperties.ProtocolMetadata.GetSubtype<FInstallBundlePluginProtocolMetaData>().InstallBundles;
+			FInstallBundlePluginProtocolMetaData& ProtocolMetadata = StateProperties.ProtocolMetadata.GetSubtype<FInstallBundlePluginProtocolMetaData>();
 
-			TValueOrError<FInstallBundleCombinedInstallState, EInstallBundleResult> MaybeInstallState = BundleManager->GetInstallStateSynchronous(InstallBundles, false);
+			const bool bAddDependencies = true;
+			TValueOrError<FInstallBundleCombinedInstallState, EInstallBundleResult> MaybeInstallState = BundleManager->GetInstallStateSynchronous(
+				ProtocolMetadata.InstallBundles, bAddDependencies);
 			if (MaybeInstallState.HasError())
 			{
 				StateStatus.SetTransitionError(EGameFeaturePluginState::ErrorCheckingStatus, GetErrorResult(TEXT("BundleManager_Failed_GetInstallState")));
@@ -820,7 +968,17 @@ struct FGameFeaturePluginState_CheckingStatus : public FGameFeaturePluginState
 			}
 
 			const FInstallBundleCombinedInstallState& InstallState = MaybeInstallState.GetValue();
-			bIsAvailable = Algo::AllOf(InstallBundles, [&InstallState](FName BundleName) { return InstallState.IndividualBundleStates.Contains(BundleName); });
+			bIsAvailable = Algo::AllOf(ProtocolMetadata.InstallBundles, 
+				[&InstallState](FName BundleName) { return InstallState.IndividualBundleStates.Contains(BundleName); });
+
+			if (bIsAvailable)
+			{
+				// Update metadata with fully expanded dependency list. This can only be done after all bundles are known to be available,
+				// otherwise unavailable bundles in the URL could be stripped from the list.
+				ProtocolMetadata.InstallBundles.Empty(InstallState.IndividualBundleStates.Num());
+				InstallState.IndividualBundleStates.GetKeys(ProtocolMetadata.InstallBundles);
+				ProtocolMetadata.InstallBundlesWithAssetDependencies = InstallState.BundlesWithIoStoreOnDemand.Array();
+			}
 		}
 		else
 		{
@@ -883,7 +1041,7 @@ struct FGameFeaturePluginState_StatusKnown : public FDestinationGameFeaturePlugi
 		{
 			if (ShouldVisitUninstallStateBeforeTerminal())
 			{
-				StateStatus.SetTransition(EGameFeaturePluginState::Uninstalling);	
+				StateStatus.SetTransition(EGameFeaturePluginState::Uninstalling);
 			}
 			else
 			{
@@ -1006,8 +1164,17 @@ struct FBaseDataReleaseGameFeaturePluginState : public FGameFeaturePluginState
 			return;
 		}
 
-		EInstallBundleReleaseRequestFlags ReleaseFlags = GetReleaseRequestFlags();
-		TValueOrError<FInstallBundleReleaseRequestInfo, EInstallBundleResult> MaybeRequestInfo = UGameFeaturesSubsystem::Get().ReleaseBundle(StateProperties.PluginName, GetInstallBundles(), ReleaseFlags);
+		TSharedPtr<IInstallBundleManager> BundleManager = IInstallBundleManager::GetPlatformInstallBundleManager();
+		check(BundleManager.IsValid());
+
+		TArray<FName> InstallBundlesToRelease = UE::GameFeatures::GFPSharedInstallTracker.Release(
+			StateProperties.PluginName, UE::GameFeatures::EGFPInstallLevel::Download, GetInstallBundles());
+
+		// Always set ExplicitRemoveList, GFPSharedInstallTracker has filtered out shared dependencies
+		EInstallBundleReleaseRequestFlags ReleaseFlags = GetReleaseRequestFlags() | EInstallBundleReleaseRequestFlags::ExplicitRemoveList;
+		TValueOrError<FInstallBundleReleaseRequestInfo, EInstallBundleResult> MaybeRequestInfo = BundleManager->RequestReleaseContent(
+			InstallBundlesToRelease, ReleaseFlags);
+
 		if (MaybeRequestInfo.HasError())
 		{
 			const FStringView ShortUrl = StateProperties.PluginIdentifier.GetIdentifyingString();
@@ -1095,8 +1262,8 @@ struct FBaseDataReleaseGameFeaturePluginState : public FGameFeaturePluginState
 	/** Determine what kind of release request flags we submit */
 	virtual EInstallBundleReleaseRequestFlags GetReleaseRequestFlags() const
 	{
-		EInstallBundleReleaseRequestFlags ReleaseFlags = StateProperties.ProtocolOptions.GetSubtype<FInstallBundlePluginProtocolOptions>().ReleaseInstallBundleFlags;
-		ReleaseFlags &= ~EInstallBundleReleaseRequestFlags::SkipReleaseUnmountOnly; // Release state MUST release
+		// Always set ExplicitRemoveList, GFPSharedInstallTracker has filtered out shared dependencies
+		EInstallBundleReleaseRequestFlags ReleaseFlags = EInstallBundleReleaseRequestFlags::ExplicitRemoveList;
 		return ReleaseFlags;
 	}
 
@@ -1199,35 +1366,39 @@ struct FGameFeaturePluginState_Releasing : public FBaseDataReleaseGameFeaturePlu
 	{
 		return EGameFeaturePluginState::ErrorManagingData;
 	}
-
-	//Overriden to ensure we don't call with RemoveFilesIfPossible early, that should
-	//only be called during Uninstall
-	virtual EInstallBundleReleaseRequestFlags GetReleaseRequestFlags() const override
-	{
-		const EInstallBundleReleaseRequestFlags BaseFlags = FBaseDataReleaseGameFeaturePluginState::GetReleaseRequestFlags();
-		return (BaseFlags & ~(EInstallBundleReleaseRequestFlags::RemoveFilesIfPossible));
-	}
 };
 
-struct FBaseDownloadGameFeaturePluginState : public FGameFeaturePluginState
+struct FGameFeaturePluginState_Downloading : public FGameFeaturePluginState
 {
-	FBaseDownloadGameFeaturePluginState(FGameFeaturePluginStateMachineProperties& InStateProperties)
+	FGameFeaturePluginState_Downloading(FGameFeaturePluginStateMachineProperties& InStateProperties)
 		: FGameFeaturePluginState(InStateProperties)
 		, Result(MakeValue())
 	{}
 
-	virtual ~FBaseDownloadGameFeaturePluginState()
+	virtual ~FGameFeaturePluginState_Downloading()
 	{
 		Cleanup();
 	}
 
 	UE::GameFeatures::FResult Result;
 	bool bPluginDownloaded = false;
+
 	TArray<FName> PendingBundleDownloads;
 	TUniquePtr<FInstallBundleCombinedProgressTracker> ProgressTracker;
 	FTSTicker::FDelegateHandle ProgressUpdateHandle;
 	FDelegateHandle GotContentStateHandle;
 
+	UE::IoStore::FIoStoreOnDemandModule* IoStoreOnDemandModule = nullptr;
+
+	// Required for callback lifetime safety
+	struct FIoStoreOnDemandContext
+	{
+		int32 PendingInstalls = 0;
+		UE::IoStore::FOnDemandCancellationToken CancelToken;
+		bool bStateValid = true;
+	};
+	TSharedPtr<FIoStoreOnDemandContext> IoStoreOnDemandContext;
+
 	void Cleanup()
 	{
 		if (ProgressUpdateHandle.IsValid())
@@ -1253,6 +1424,14 @@ struct FBaseDownloadGameFeaturePluginState : public FGameFeaturePluginState
 		bPluginDownloaded = false;
 		PendingBundleDownloads.Empty();
 		ProgressTracker = nullptr;
+		IoStoreOnDemandModule = nullptr;
+
+		if (IoStoreOnDemandContext)
+		{
+			IoStoreOnDemandContext->CancelToken.Cancel();
+			IoStoreOnDemandContext->bStateValid = false;
+			IoStoreOnDemandContext = nullptr;
+		}
 	}
 
 	void OnGotContentState(FInstallBundleCombinedContentState BundleContentState)
@@ -1268,25 +1447,19 @@ struct FBaseDownloadGameFeaturePluginState : public FGameFeaturePluginState
 			return;
 		}
 
-		// check to verify if the bundle(s) we need is already UpToDate
-		if (BundleContentState.GetAllBundlesHaveState(EInstallBundleInstallState::UpToDate))
-		{
-			bPluginDownloaded = true;
-			UpdateProgress(1.0f);
-			UpdateStateMachineImmediate();
-			return;
-		}
-
+		const FInstallBundlePluginProtocolMetaData& MetaData = StateProperties.ProtocolMetadata.GetSubtype<FInstallBundlePluginProtocolMetaData>();
 		const FInstallBundlePluginProtocolOptions& Options = StateProperties.ProtocolOptions.GetSubtype<FInstallBundlePluginProtocolOptions>();
 
-		// not up to date, check to see if we allow downloading
-		if (Options.bDoNotDownload)
+		// if not up to date, check to see if we allow downloading
+		if (Options.bDoNotDownload && !BundleContentState.GetAllBundlesHaveState(EInstallBundleInstallState::UpToDate))
 		{
 			Result = GetErrorResult(TEXT("GFPStateMachine.DownloadNotAllowed"));
 			UpdateStateMachineImmediate();
 			return;
 		}
 
+		UE::GameFeatures::GFPSharedInstallTracker.AddBundleRefs(StateProperties.PluginName, UE::GameFeatures::EGFPInstallLevel::Download, MetaData.InstallBundles);
+
 		const TConstArrayView<FName> InstallBundles = GetInstallBundles();
 		const EInstallBundleRequestFlags InstallFlags = GetRequestFlags();
 		TValueOrError<FInstallBundleRequestInfo, EInstallBundleResult> MaybeRequestInfo = BundleManager->RequestUpdateContent(InstallBundles, InstallFlags);
@@ -1314,21 +1487,20 @@ struct FBaseDownloadGameFeaturePluginState : public FGameFeaturePluginState
 
 		if (RequestInfo.BundlesEnqueued.Num() == 0)
 		{
-			bPluginDownloaded = true;
 			UpdateProgress(1.0f);
-			UpdateStateMachineImmediate();
+			InstallIoStoreOnDemandContent();
 		}
 		else
 		{
 			PendingBundleDownloads = MoveTemp(RequestInfo.BundlesEnqueued);
-			IInstallBundleManager::InstallBundleCompleteDelegate.AddRaw(this, &FBaseDownloadGameFeaturePluginState::OnInstallBundleCompleted);
-			IInstallBundleManager::PausedBundleDelegate.AddRaw(this, &FBaseDownloadGameFeaturePluginState::OnInstallBundlePaused);
+			IInstallBundleManager::InstallBundleCompleteDelegate.AddRaw(this, &FGameFeaturePluginState_Downloading::OnInstallBundleCompleted);
+			IInstallBundleManager::PausedBundleDelegate.AddRaw(this, &FGameFeaturePluginState_Downloading::OnInstallBundlePaused);
 
 			ProgressTracker = MakeUnique<FInstallBundleCombinedProgressTracker>(false);
 			ProgressTracker->SetBundlesToTrackFromContentState(BundleContentState, PendingBundleDownloads);
 
 			ProgressUpdateHandle = FTSTicker::GetCoreTicker().AddTicker(
-				FTickerDelegate::CreateRaw(this, &FBaseDownloadGameFeaturePluginState::OnUpdateProgress)/*, 0.1f*/);
+				FTickerDelegate::CreateRaw(this, &FGameFeaturePluginState_Downloading::OnUpdateProgress)/*, 0.1f*/);
 
 			//If this setting is flipped then we should immediately request to pause downloads.
 			//We still generate the downloads so that we have an accurate PendingBundleDownloads list
@@ -1358,8 +1530,7 @@ struct FBaseDownloadGameFeaturePluginState : public FGameFeaturePluginState
 
 			if (BundleResult.Result != EInstallBundleResult::UserCancelledError)
 			{
-				TSharedPtr<IInstallBundleManager> BundleManager = IInstallBundleManager::GetPlatformInstallBundleManager();
-				BundleManager->CancelUpdateContent(PendingBundleDownloads);
+				TryCancelState();
 			}
 		}
 
@@ -1368,14 +1539,75 @@ struct FBaseDownloadGameFeaturePluginState : public FGameFeaturePluginState
 			return;
 		}
 
-		if (Result.HasValue())
+		OnUpdateProgress(0.0f);
+
+		InstallIoStoreOnDemandContent();
+	}
+
+	void InstallIoStoreOnDemandContent()
+	{
+		const FInstallBundlePluginProtocolMetaData& MetaData = StateProperties.ProtocolMetadata.GetSubtype<FInstallBundlePluginProtocolMetaData>();
+		if (MetaData.InstallBundlesWithAssetDependencies.IsEmpty())
 		{
-			bPluginDownloaded = true;
+			bPluginDownloaded = Result.HasValue();
+			UpdateStateMachineImmediate();
+			return;
 		}
 
-		OnUpdateProgress(0.0f);
+		IoStoreOnDemandModule = FModuleManager::Get().GetModulePtr<UE::IoStore::FIoStoreOnDemandModule>(TEXT("IoStoreOnDemand"));
+		if (!IoStoreOnDemandModule && !Result.HasError())
+		{
+			Result = GetErrorResult(TEXT("IoStoreOnDemand.ModuleNotFound"));
+			UpdateStateMachineImmediate();
+			return;
+		}
 
-		UpdateStateMachineImmediate();
+		IoStoreOnDemandContext = MakeShared<FIoStoreOnDemandContext>();
+		IoStoreOnDemandContext->PendingInstalls = MetaData.InstallBundlesWithAssetDependencies.Num();
+
+		for (const FName InstallBundle : MetaData.InstallBundlesWithAssetDependencies)
+		{
+			UE::IoStore::FOnDemandInstallArgs InstallArgs;
+			InstallArgs.MountId = InstallBundle.ToString();
+			InstallArgs.TagSets.Emplace(TEXTVIEW("mount")); // May not be a real tagset, but this will install all the mandatory untagged chunks
+			InstallArgs.Options |= UE::IoStore::EOnDemandInstallOptions::CallbackOnGameThread;
+			
+			InstallArgs.ContentHandle = UE::GameFeatures::GFPSharedInstallTracker.AddOnDemandContentHandle(
+				InstallBundle, UE::GameFeatures::EGFPInstallLevel::Download);
+			check(InstallArgs.ContentHandle.IsValid());
+
+			// This should be pretty small, so not going to worry about progress here.
+			IoStoreOnDemandModule->Install(MoveTemp(InstallArgs),
+				[this, LambdaOnDemandContext = IoStoreOnDemandContext](const UE::IoStore::FOnDemandInstallResult& OnDemandInstallResult)
+				{
+					if (!LambdaOnDemandContext->bStateValid)
+					{
+						// Owning state got cleaned up, bail
+						return;
+					}
+
+					if (!OnDemandInstallResult.Status.IsOk() && !Result.HasError())
+					{
+						const FString StatusString = OnDemandInstallResult.Status.ToString();
+						FString ErrorCode = GetIoErrorText(OnDemandInstallResult.Status.GetErrorCode());
+						ErrorCode.ReplaceCharInline(TEXT(' '), TEXT('_'), ESearchCase::CaseSensitive);
+
+						// IOStore errors aren't localized 
+						Result = GetErrorResult(TEXT("IoStoreOnDemand.OnComplete."), ErrorCode, FText::AsCultureInvariant(StatusString));
+						if (!LambdaOnDemandContext->CancelToken.IsCanceled())
+						{
+							TryCancelState();
+						}
+					}
+
+					--IoStoreOnDemandContext->PendingInstalls;
+					if (IoStoreOnDemandContext->PendingInstalls == 0)
+					{
+						bPluginDownloaded = Result.HasValue();
+						UpdateStateMachineImmediate();
+					}
+				}, nullptr, &IoStoreOnDemandContext->CancelToken);
+		}
 	}
 
 	bool OnUpdateProgress(float dts)
@@ -1438,31 +1670,34 @@ struct FBaseDownloadGameFeaturePluginState : public FGameFeaturePluginState
 	{
 		Cleanup();
 
-		if (!ShouldDownloadContent())
+		if (!ensure(ShouldDownloadContent()))
 		{
 			bPluginDownloaded = true;
 			UpdateProgress(1.0f);
 			return;
 		}
 
+		UGameFeaturesSubsystem::Get().OnGameFeatureDownloading(StateProperties.PluginName, StateProperties.PluginIdentifier);
+
 		ensureMsgf(AllowAsyncLoading(), TEXT("FGameFeaturePluginState::AllowAsyncLoading is false while attempting to download GFP data."));
 
 		TSharedPtr<IInstallBundleManager> BundleManager = IInstallBundleManager::GetPlatformInstallBundleManager();
 		const TConstArrayView<FName> InstallBundles = GetInstallBundles();
 
+		const bool bAddDependencies = false; // We already got all dependencies in the CheckingStatus state
 		if (InstallBundles.Num() > 1)
 		{
-			GotContentStateHandle = BundleManager->GetContentState(InstallBundles, EInstallBundleGetContentStateFlags::None, true,
-				FInstallBundleGetContentStateDelegate::CreateRaw(this, &FBaseDownloadGameFeaturePluginState::OnGotContentState));
+			GotContentStateHandle = BundleManager->GetContentState(InstallBundles, EInstallBundleGetContentStateFlags::None, bAddDependencies,
+				FInstallBundleGetContentStateDelegate::CreateRaw(this, &FGameFeaturePluginState_Downloading::OnGotContentState));
 		}
 		else
 		{
 			// We usualy only have a use case of one bundle per GFP and we only care about relative weighting here, so we don't need 
 			// any of the other content state metadata. We can just assume the weight is 1.0 and not have to wait for the full 
 			// async call to get the rest of the metadata.
-			TValueOrError<FInstallBundleCombinedInstallState, EInstallBundleResult> MaybeInstallState = BundleManager->GetInstallStateSynchronous(InstallBundles, true);
+			TValueOrError<FInstallBundleCombinedInstallState, EInstallBundleResult> MaybeInstallState = BundleManager->GetInstallStateSynchronous(InstallBundles, bAddDependencies);
 			check(MaybeInstallState.HasValue());
-			const FInstallBundleCombinedInstallState& InstallState = MaybeInstallState.GetValue();
+			FInstallBundleCombinedInstallState& InstallState = MaybeInstallState.GetValue();
 			FInstallBundleCombinedContentState HackContentState;
 			HackContentState.IndividualBundleStates.Reserve(InstallState.IndividualBundleStates.Num());
 			for (const TPair<FName, EInstallBundleInstallState>& Pair : InstallState.IndividualBundleStates)
@@ -1471,6 +1706,7 @@ struct FBaseDownloadGameFeaturePluginState : public FGameFeaturePluginState
 				BundleContentState.State = Pair.Value;
 				BundleContentState.Weight = 1.0f;
 			}
+			HackContentState.BundlesWithIoStoreOnDemand = MoveTemp(InstallState.BundlesWithIoStoreOnDemand);
 			OnGotContentState(MoveTemp(HackContentState));
 		}
 	}
@@ -1488,8 +1724,7 @@ struct FBaseDownloadGameFeaturePluginState : public FGameFeaturePluginState
 		{
 			return;
 		}
-		
-		UGameFeaturesSubsystem::Get().OnGameFeatureDownloaded(StateProperties.PluginIdentifier);
+
 		const EGameFeaturePluginState SuccessState = GetSuccessTransitionState();
 		StateStatus.SetTransition(SuccessState);
 	}
@@ -1500,6 +1735,10 @@ struct FBaseDownloadGameFeaturePluginState : public FGameFeaturePluginState
 		{
 			TSharedPtr<IInstallBundleManager> BundleManager = IInstallBundleManager::GetPlatformInstallBundleManager();
 			BundleManager->CancelUpdateContent(PendingBundleDownloads);
+			if (IoStoreOnDemandContext)
+			{
+				IoStoreOnDemandContext->CancelToken.Cancel();
+			}
 		}
 	}
 
@@ -1581,47 +1820,22 @@ struct FBaseDownloadGameFeaturePluginState : public FGameFeaturePluginState
 	}
 
 	/** Determine what kind of request flags we submit */
-	virtual EInstallBundleRequestFlags GetRequestFlags() const
-	{
-		EInstallBundleRequestFlags InstallFlags = StateProperties.ProtocolOptions.GetSubtype<FInstallBundlePluginProtocolOptions>().InstallBundleFlags;
-		return InstallFlags;
-	}
-
-	/** Determines what state you transition to in the event of a success or failure to release content */
-	virtual EGameFeaturePluginState GetSuccessTransitionState() const = 0;
-	virtual EGameFeaturePluginState GetFailureTransitionState() const = 0;
-};
-
-struct FGameFeaturePluginState_Downloading : public FBaseDownloadGameFeaturePluginState
-{
-	FGameFeaturePluginState_Downloading(FGameFeaturePluginStateMachineProperties& InStateProperties)
-		: FBaseDownloadGameFeaturePluginState(InStateProperties)
-	{}
-
-	virtual void BeginState() override
-	{
-		check(ShouldDownloadContent());
-
-		UGameFeaturesSubsystem::Get().OnGameFeatureDownloading(StateProperties.PluginName, StateProperties.PluginIdentifier);
-
-		FBaseDownloadGameFeaturePluginState::BeginState();
-	}
-	
-	virtual EInstallBundleRequestFlags GetRequestFlags() const override
+	EInstallBundleRequestFlags GetRequestFlags() const
 	{
 		//Pull our InstallFlags from the Options, but also make sure SkipMount is set as there is a separate mounting step that will re-request this
 		//without SkipMount and then mount the data, this allows us to pre-download data without mounting it
-		EInstallBundleRequestFlags InstallFlags = FBaseDownloadGameFeaturePluginState::GetRequestFlags();
+		EInstallBundleRequestFlags InstallFlags = StateProperties.ProtocolOptions.GetSubtype<FInstallBundlePluginProtocolOptions>().InstallBundleFlags;
 		InstallFlags |= EInstallBundleRequestFlags::SkipMount;
 		return InstallFlags;
 	}
 
-	virtual EGameFeaturePluginState GetSuccessTransitionState() const override
+	/** Determines what state you transition to in the event of a success or failure to release content */
+	EGameFeaturePluginState GetSuccessTransitionState() const
 	{
 		return EGameFeaturePluginState::Installed;
 	}
 
-	virtual EGameFeaturePluginState GetFailureTransitionState() const override
+	EGameFeaturePluginState GetFailureTransitionState() const
 	{
 		return EGameFeaturePluginState::ErrorManagingData;
 	}
@@ -1744,13 +1958,19 @@ struct FGameFeaturePluginState_Unmounting : public FGameFeaturePluginState
 			return;
 		}
 
+		TSharedPtr<IInstallBundleManager> BundleManager = IInstallBundleManager::GetPlatformInstallBundleManager();
+		check(BundleManager.IsValid());
+
 		const TArray<FName>& InstallBundles = StateProperties.ProtocolMetadata.GetSubtype<FInstallBundlePluginProtocolMetaData>().InstallBundles;
 
-		EInstallBundleReleaseRequestFlags ReleaseFlags = StateProperties.ProtocolOptions.GetSubtype<FInstallBundlePluginProtocolOptions>().ReleaseInstallBundleFlags;
-		ReleaseFlags |= EInstallBundleReleaseRequestFlags::SkipReleaseUnmountOnly;
-		//Make sure we don't remove files here early, that should only be done in Uninstalling
-		ReleaseFlags &= ~(EInstallBundleReleaseRequestFlags::RemoveFilesIfPossible);
-		TValueOrError<FInstallBundleReleaseRequestInfo, EInstallBundleResult> MaybeRequestInfo = UGameFeaturesSubsystem::Get().UnmountBundle(StateProperties.PluginName, InstallBundles, ReleaseFlags);
+		TArray<FName> InstallBundlesToRelease = UE::GameFeatures::GFPSharedInstallTracker.Release(
+			StateProperties.PluginName, UE::GameFeatures::EGFPInstallLevel::Mount, InstallBundles);
+
+		EInstallBundleReleaseRequestFlags ReleaseFlags = 
+			EInstallBundleReleaseRequestFlags::SkipReleaseUnmountOnly |
+			EInstallBundleReleaseRequestFlags::ExplicitRemoveList; // Always set ExplicitRemoveList, GFPSharedInstallTracker has filtered out shared dependencies
+		TValueOrError<FInstallBundleReleaseRequestInfo, EInstallBundleResult> MaybeRequestInfo = BundleManager->RequestReleaseContent(
+			InstallBundlesToRelease, ReleaseFlags);
 
 		if (MaybeRequestInfo.HasError())
 		{
@@ -1975,6 +2195,8 @@ struct FGameFeaturePluginState_Mounting : public FGameFeaturePluginState
 		const FInstallBundlePluginProtocolMetaData& MetaData = StateProperties.ProtocolMetadata.GetSubtype<FInstallBundlePluginProtocolMetaData>();
 		const TArray<FName>& InstallBundles = MetaData.InstallBundles;
 
+		UE::GameFeatures::GFPSharedInstallTracker.AddBundleRefs(StateProperties.PluginName, UE::GameFeatures::EGFPInstallLevel::Mount, InstallBundles);
+
 		const FInstallBundlePluginProtocolOptions& Options = StateProperties.ProtocolOptions.GetSubtype<FInstallBundlePluginProtocolOptions>();
 		const EInstallBundleRequestFlags InstallFlags = UseAsyncLoading() ?
 			(Options.InstallBundleFlags | EInstallBundleRequestFlags::AsyncMount) : Options.InstallBundleFlags;
@@ -2370,39 +2592,28 @@ struct FGameFeaturePluginState_WaitingForDependencies : public FTransitionDepend
 	}
 };
 
-struct FGameFeaturePluginState_AssetDependencyStreamOut : public FBaseDataReleaseGameFeaturePluginState
+struct FGameFeaturePluginState_AssetDependencyStreamOut : public FGameFeaturePluginState
 {
 	FGameFeaturePluginState_AssetDependencyStreamOut(FGameFeaturePluginStateMachineProperties& InStateProperties)
-		: FBaseDataReleaseGameFeaturePluginState(InStateProperties)
+		: FGameFeaturePluginState(InStateProperties)
 	{}
 
-	virtual EGameFeaturePluginState GetSuccessTransitionState() const override
-	{
-		return EGameFeaturePluginState::Unmounting;
-	}
-
-	virtual EGameFeaturePluginState GetFailureTransitionState() const override
-	{
-		return EGameFeaturePluginState::ErrorAssetDependencyStreaming;
-	}
-
-	virtual TConstArrayView<FName> GetInstallBundles() override
+	virtual void BeginState() override
 	{
-		TConstArrayView<FName> Ret;
-		if (ShouldReleaseContent())
+		if (StateProperties.GetPluginProtocol() != EGameFeaturePluginProtocol::InstallBundle)
 		{
-			Ret = StateProperties.ProtocolMetadata.GetSubtype<FInstallBundlePluginProtocolMetaData>().AssetDependencyBundles;
+			return;
 		}
-		return Ret;
+
+		const FInstallBundlePluginProtocolMetaData& MetaData = StateProperties.ProtocolMetadata.GetSubtype<FInstallBundlePluginProtocolMetaData>();
+
+		UE::GameFeatures::GFPSharedInstallTracker.Release(
+			StateProperties.PluginName, UE::GameFeatures::EGFPInstallLevel::AssetStream, MetaData.InstallBundlesWithAssetDependencies);
 	}
 
-	// There's no explict uninstall set for asset dependencies, so set EInstallBundleReleaseRequestFlags::RemoveFilesIfPossible
-	// to be 'correct'. In most cases, we'd expect to store these with a cache of some sort, so this exact flag probably wont matter
-	// in practice.
-	virtual EInstallBundleReleaseRequestFlags GetReleaseRequestFlags() const override
+	virtual void UpdateState(FGameFeaturePluginStateStatus& StateStatus) override
 	{
-		const EInstallBundleReleaseRequestFlags BaseFlags = FBaseDataReleaseGameFeaturePluginState::GetReleaseRequestFlags();
-		return (BaseFlags | EInstallBundleReleaseRequestFlags::RemoveFilesIfPossible);
+		StateStatus.SetTransition(EGameFeaturePluginState::Unmounting);
 	}
 };
 
@@ -2425,40 +2636,200 @@ struct FGameFeaturePluginState_ErrorAssetDependencyStreaming : public FErrorGame
 	}
 };
 
-struct FGameFeaturePluginState_AssetDependencyStreaming : public FBaseDownloadGameFeaturePluginState
+struct FGameFeaturePluginState_AssetDependencyStreaming : public FGameFeaturePluginState
 {
 	FGameFeaturePluginState_AssetDependencyStreaming(FGameFeaturePluginStateMachineProperties& InStateProperties)
-		: FBaseDownloadGameFeaturePluginState(InStateProperties)
+		: FGameFeaturePluginState(InStateProperties)
 	{}
 
-	virtual TConstArrayView<FName> GetInstallBundles() override
+	virtual ~FGameFeaturePluginState_AssetDependencyStreaming()
 	{
-		TConstArrayView<FName> Ret;
-		if (ShouldDownloadContent())
+		Cleanup();
+	}
+
+	struct FIoStoreOnDemandProgress
+	{
+		FName InstallBundle;
+		float Progress = 0.0f;
+	};
+
+	// Required for callback lifetime safety
+	struct FIoStoreOnDemandContext
+	{
+		TArray<FIoStoreOnDemandProgress> Progress;
+		int32 PendingInstalls = 0;
+		UE::IoStore::FOnDemandCancellationToken CancelToken;
+		bool bStateValid = true;
+	};
+
+	UE::IoStore::FIoStoreOnDemandModule* IoStoreOnDemandModule = nullptr;
+	TSharedPtr<FIoStoreOnDemandContext> IoStoreOnDemandContext;
+
+	UE::GameFeatures::FResult Result{ MakeValue() };
+	bool bComplete = false;
+
+	void Cleanup()
+	{
+		Result = MakeValue();
+		bComplete = false;
+
+		IoStoreOnDemandModule = nullptr;
+
+		if (IoStoreOnDemandContext)
 		{
-			Ret = StateProperties.ProtocolMetadata.GetSubtype<FInstallBundlePluginProtocolMetaData>().AssetDependencyBundles;
+			IoStoreOnDemandContext->CancelToken.Cancel();
+			IoStoreOnDemandContext->bStateValid = false;
+			IoStoreOnDemandContext = nullptr;
 		}
-		return Ret;
 	}
 
-	virtual EInstallBundleRequestFlags GetRequestFlags() const override
+	virtual void BeginState() override
 	{
-		EInstallBundleRequestFlags InstallFlags = FBaseDownloadGameFeaturePluginState::GetRequestFlags();
-		if (UseAsyncLoading())
+		Cleanup();
+
+		if (StateProperties.GetPluginProtocol() != EGameFeaturePluginProtocol::InstallBundle)
 		{
-			InstallFlags |= EInstallBundleRequestFlags::AsyncMount;
+			bComplete = true;
+			return;
+		}
+
+		const FInstallBundlePluginProtocolMetaData& MetaData = StateProperties.ProtocolMetadata.GetSubtype<FInstallBundlePluginProtocolMetaData>();
+
+		if (MetaData.InstallBundlesWithAssetDependencies.IsEmpty())
+		{
+			bComplete = true;
+			return;
+		}
+
+		IoStoreOnDemandModule = FModuleManager::Get().GetModulePtr<UE::IoStore::FIoStoreOnDemandModule>(TEXT("IoStoreOnDemand"));
+		if (!IoStoreOnDemandModule)
+		{
+			Result = GetErrorResult(TEXT("IoStoreOnDemand.ModuleNotFound"));
+			return;
+		}
+
+		TValueOrError<TArray<EStreamingAssetInstallMode>, FString> MaybeInstallModes = 
+			UGameFeaturesSubsystem::Get().GetPolicy().GetStreamingAssetInstallModes(
+				StateProperties.PluginIdentifier.GetFullPluginURL(), MetaData.InstallBundlesWithAssetDependencies);
+
+		if (MaybeInstallModes.HasError())
+		{
+			Result = GetErrorResult(TEXT("IoStoreOnDemand.InstallMode"), MaybeInstallModes.GetError());
+			return;
+		}
+
+		UE::GameFeatures::GFPSharedInstallTracker.AddBundleRefs(
+			StateProperties.PluginName, UE::GameFeatures::EGFPInstallLevel::AssetStream, MetaData.InstallBundlesWithAssetDependencies);
+
+		IoStoreOnDemandContext = MakeShared<FIoStoreOnDemandContext>();
+		IoStoreOnDemandContext->PendingInstalls = MetaData.InstallBundlesWithAssetDependencies.Num();
+		IoStoreOnDemandContext->Progress.Reserve(MetaData.InstallBundlesWithAssetDependencies.Num());
+
+		const TArray<EStreamingAssetInstallMode>& InstallModes = MaybeInstallModes.GetValue();
+		for (int i = 0; const FName InstallBundle : MetaData.InstallBundlesWithAssetDependencies)
+		{
+			IoStoreOnDemandContext->Progress.Emplace(FIoStoreOnDemandProgress{InstallBundle});
+
+			const EStreamingAssetInstallMode InstallMode = InstallModes[i++];
+
+			UE::IoStore::FOnDemandInstallArgs InstallArgs;
+			InstallArgs.MountId = InstallBundle.ToString();
+			if (InstallMode == EStreamingAssetInstallMode::GfpRequiredOnly)
+			{
+				InstallArgs.TagSets.Emplace(TEXTVIEW("required"));
+			}			
+			InstallArgs.Options |= UE::IoStore::EOnDemandInstallOptions::CallbackOnGameThread;
+
+			InstallArgs.ContentHandle = UE::GameFeatures::GFPSharedInstallTracker.AddOnDemandContentHandle(
+				InstallBundle, UE::GameFeatures::EGFPInstallLevel::AssetStream);
+
+			IoStoreOnDemandModule->Install(MoveTemp(InstallArgs),
+				// On Complete
+				[this, LambdaOnDemandContext = IoStoreOnDemandContext](const UE::IoStore::FOnDemandInstallResult& OnDemandInstallResult)
+				{
+					if (!LambdaOnDemandContext->bStateValid)
+					{
+						// Owning state got cleaned up, bail
+						return;
+					}
+
+					if (!OnDemandInstallResult.Status.IsOk() && !Result.HasError())
+					{
+						const FString StatusString = OnDemandInstallResult.Status.ToString();
+						FString ErrorCode = GetIoErrorText(OnDemandInstallResult.Status.GetErrorCode());
+						ErrorCode.ReplaceCharInline(TEXT(' '), TEXT('_'), ESearchCase::CaseSensitive);
+
+						// IOStore errors aren't localized 
+						Result = GetErrorResult(TEXT("IoStoreOnDemand.OnComplete."), ErrorCode, FText::AsCultureInvariant(StatusString));
+						if (!LambdaOnDemandContext->CancelToken.IsCanceled())
+						{
+							TryCancelState();
+						}
+					}
+
+					--IoStoreOnDemandContext->PendingInstalls;
+					if (IoStoreOnDemandContext->PendingInstalls == 0)
+					{
+						bComplete = true;
+						UpdateStateMachineImmediate();
+					}
+
+				},
+				// On Progress
+				[this, LambdaOnDemandContext = IoStoreOnDemandContext, InstallBundle](const UE::IoStore::FOnDemandInstallProgress& Progress)
+				{
+					if (!LambdaOnDemandContext->bStateValid)
+					{
+						// Owning state got cleaned up, bail
+						return;
+					}
+
+					const double RequestProgress = (Progress.TotalInstallSize > 0) ?
+						FMath::Clamp(double(Progress.CurrentInstallSize) / Progress.TotalInstallSize, 0.0, 1.0) :
+						0;
+
+					FIoStoreOnDemandProgress* MyProgress = Algo::FindBy(
+						LambdaOnDemandContext->Progress, InstallBundle, &FIoStoreOnDemandProgress::InstallBundle);
+					check(MyProgress);
+					MyProgress->Progress = RequestProgress;
+
+					const float SumProgress = Algo::TransformAccumulate(
+						LambdaOnDemandContext->Progress, [](const FIoStoreOnDemandProgress& In) { return In.Progress; }, 0.0f);
+					const float OverallProgress = SumProgress / LambdaOnDemandContext->Progress.Num();
+
+					StateProperties.OnFeatureStateProgressUpdate.ExecuteIfBound(OverallProgress);
+				},
+				&IoStoreOnDemandContext->CancelToken);
 		}
-		return InstallFlags;
 	}
 
-	virtual EGameFeaturePluginState GetSuccessTransitionState() const override
+	virtual void UpdateState(FGameFeaturePluginStateStatus& StateStatus) override
 	{
-		return EGameFeaturePluginState::Registering;
+		if (!Result.HasValue())
+		{
+			StateStatus.SetTransitionError(EGameFeaturePluginState::ErrorAssetDependencyStreaming, Result);
+			return;
+		}
+
+		if (!bComplete)
+		{
+			return;
+		}
+
+		StateStatus.SetTransition(EGameFeaturePluginState::Registering);
 	}
 
-	virtual EGameFeaturePluginState GetFailureTransitionState() const override
+	virtual void EndState() override
 	{
-		return EGameFeaturePluginState::ErrorAssetDependencyStreaming;
+		Cleanup();
+	}
+
+	virtual void TryCancelState() override
+	{
+		if (IoStoreOnDemandContext)
+		{
+			IoStoreOnDemandContext->CancelToken.Cancel();
+		}
 	}
 };
 
@@ -2471,7 +2842,7 @@ struct FGameFeaturePluginState_Unregistering : public FGameFeaturePluginState
 	bool bRequestedUnloadPluginAssets = false;
 #endif //if WITH_EDITOR
 
-	virtual void BeginState()
+	virtual void BeginState() override
 	{
 		bHasUnloaded = false;
 #if WITH_EDITOR
@@ -3517,7 +3888,7 @@ bool UGameFeaturePluginStateMachine::HasAssetStreamingDependencies() const
 	if (StateProperties.ProtocolMetadata.HasSubtype<FInstallBundlePluginProtocolMetaData>())
 	{
 		const FInstallBundlePluginProtocolMetaData& ProtocolData = StateProperties.ProtocolMetadata.GetSubtype<FInstallBundlePluginProtocolMetaData>();
-		return !ProtocolData.AssetDependencyBundles.IsEmpty();
+		return !ProtocolData.InstallBundlesWithAssetDependencies.IsEmpty();
 	}
 	return false;
 }
@@ -3811,18 +4182,8 @@ TValueOrError<FInstallBundlePluginProtocolMetaData, FString> FInstallBundlePlugi
 		return MakeError(TEXTVIEW("Bad_PluginURL"));
 	}
 
-	// Get any additional bundles for streamed assets
-	TValueOrError<TArray<FName>, FString> MaybeStreamingAssetInstallBundles =
-		UGameFeaturesSubsystem::Get().GetPolicy().GetStreamingAssetInstallBundles(URLOptionsString);
-	if (MaybeStreamingAssetInstallBundles.HasError())
-	{
-		UE_LOG(LogGameFeatures, Error, TEXT("Failed to get streaming asset install bundles from GFP policy. URL %.*s"), URLOptionsString.Len(), URLOptionsString.GetData());
-		return MakeError(MaybeStreamingAssetInstallBundles.StealError());
-	}
-
 	FInstallBundlePluginProtocolMetaData Ret;
 	Ret.InstallBundles = MoveTemp(InstallBundles);
-	Ret.AssetDependencyBundles = MaybeStreamingAssetInstallBundles.StealValue();
 
 	return MakeValue<FInstallBundlePluginProtocolMetaData>(MoveTemp(Ret));
 }
diff --git a/Engine/Plugins/Runtime/GameFeatures/Source/GameFeatures/Private/GameFeaturePluginStateMachine.h b/Engine/Plugins/Runtime/GameFeatures/Source/GameFeatures/Private/GameFeaturePluginStateMachine.h
--- a/Engine/Plugins/Runtime/GameFeatures/Source/GameFeatures/Private/GameFeaturePluginStateMachine.h
+++ b/Engine/Plugins/Runtime/GameFeatures/Source/GameFeatures/Private/GameFeaturePluginStateMachine.h
@@ -67,7 +67,7 @@ Transition states are expected to transition the machine to another state after
                  |    |                               |                    |
                  |    |                               +--------------------+
                  |    |
-            +----+----v----+                            
+            +----+----v----+
             |      *       |
          ---> StatusKnown  +----------------------------------------------
          |  |              |                                 |           |
@@ -132,7 +132,7 @@ Transition states are expected to transition the machine to another state after
            |     |        |                               ^      |
         +--v-----+--+  +--v--------+                      |    +-+------------+
         |           |  |           |                      |    |      !       |
-        | Unloading |  |  Loading  <-------------------- ~----> ErrorLoading |
+        | Unloading |  |  Loading  <----------------------~----> ErrorLoading |
         |           |  |           |                      |    |              |
         +--------^--+  +--+--------+                      |    +--------------+
                  |        |                               |
@@ -237,7 +237,7 @@ struct FInstallBundlePluginProtocolMetaData
 	FInstallBundlePluginProtocolMetaData(TArray<FName> InInstallBundles) : InstallBundles(MoveTemp(InInstallBundles)) {}
 
 	TArray<FName> InstallBundles;
-	TArray<FName> AssetDependencyBundles;
+	TArray<FName> InstallBundlesWithAssetDependencies;
 
 	/** Functions to convert to/from the URL FString representation of this metadata **/
 	FString ToString() const;
@@ -291,19 +291,19 @@ struct FGameFeaturePluginStateMachineProperties
 	/** Additional options for a specific protocol. */
 	FGameFeatureProtocolOptions ProtocolOptions;
 
-	TArray<FName> AddedPrimaryAssetTypes;
+	/** The desired state during a transition. */
+	FGameFeaturePluginStateRange Destination;
 
 	/** Tracks whether or not this state machine added the plugin to the plugin manager. */
-	bool bAddedPluginToManager = false;
+	bool bAddedPluginToManager : 1 = false;
 
 	/** Was this plugin loaded using LoadBuiltInGameFeaturePlugin */
-	bool bWasLoadedAsBuiltInGameFeaturePlugin = false;
+	bool bWasLoadedAsBuiltInGameFeaturePlugin : 1 = false;
 
 	/** Whether this state machine should attempt to cancel the current transition */
-	bool bTryCancel = false;
+	bool bTryCancel : 1 = false;
 
-	/** The desired state during a transition. */
-	FGameFeaturePluginStateRange Destination;
+	TArray<FName> AddedPrimaryAssetTypes;
 
 	/** The data asset describing this game feature */
 	UPROPERTY(Transient)
diff --git a/Engine/Plugins/Runtime/GameFeatures/Source/GameFeatures/Private/GameFeaturesProjectPolicies.cpp b/Engine/Plugins/Runtime/GameFeatures/Source/GameFeatures/Private/GameFeaturesProjectPolicies.cpp
--- a/Engine/Plugins/Runtime/GameFeatures/Source/GameFeatures/Private/GameFeaturesProjectPolicies.cpp
+++ b/Engine/Plugins/Runtime/GameFeatures/Source/GameFeatures/Private/GameFeaturesProjectPolicies.cpp
@@ -105,6 +105,18 @@ TValueOrError<FString, FString> UGameFeaturesProjectPolicies::ResolvePluginDepen
 	return MakeError(TEXT("NotFound"));
 }
 
+TValueOrError<TArray<EStreamingAssetInstallMode>, FString> UGameFeaturesProjectPolicies::GetStreamingAssetInstallModes(FStringView PluginURL, TConstArrayView<FName> InstallBundleNames) const
+{
+	TArray<EStreamingAssetInstallMode> InstallModes;
+	InstallModes.Reserve(InstallBundleNames.Num());
+	while (InstallModes.Num() < InstallBundleNames.Num())
+	{
+		InstallModes.Emplace(EStreamingAssetInstallMode::Full);
+	}
+
+	return MakeValue(MoveTemp(InstallModes));
+}
+
 void UGameFeaturesProjectPolicies::ExplicitLoadGameFeaturePlugin(const FString& PluginURL, const FGameFeaturePluginLoadComplete& CompleteDelegate, const bool bActivateGameFeatures)
 {
 	if (bActivateGameFeatures)
diff --git a/Engine/Plugins/Runtime/GameFeatures/Source/GameFeatures/Private/GameFeaturesSubsystem.cpp b/Engine/Plugins/Runtime/GameFeatures/Source/GameFeatures/Private/GameFeaturesSubsystem.cpp
--- a/Engine/Plugins/Runtime/GameFeatures/Source/GameFeatures/Private/GameFeaturesSubsystem.cpp
+++ b/Engine/Plugins/Runtime/GameFeatures/Source/GameFeatures/Private/GameFeaturesSubsystem.cpp
@@ -314,14 +314,12 @@ bool FInstallBundlePluginProtocolOptions::operator==(const FInstallBundlePluginP
 {
 	return
 		InstallBundleFlags == Other.InstallBundleFlags &&
-		ReleaseInstallBundleFlags == Other.ReleaseInstallBundleFlags &&
 		bUninstallBeforeTerminate == Other.bUninstallBeforeTerminate &&
 		bUserPauseDownload == Other.bUserPauseDownload && 
 		bAllowIniLoading == Other.bAllowIniLoading &&
 		bDoNotDownload == Other.bDoNotDownload;
 }
 
-// TODO : FGameFeatureProtocolOptions - C++20 will allow inline init for bitfields
 FGameFeatureProtocolOptions::FGameFeatureProtocolOptions()
 	: bForceSyncLoading(false)
 	, bLogWarningOnForcedDependencyCreation(false)
@@ -994,23 +992,6 @@ void UGameFeaturesSubsystem::OnGameFeatureDownloading(const FString& PluginName,
 	CallbackObservers(EObserverCallback::Downloading, PluginIdentifier, &PluginName);
 }
 
-void UGameFeaturesSubsystem::OnGameFeatureDownloaded(const FGameFeaturePluginIdentifier& PluginIdentifier)
-{
-	if (UGameFeaturePluginStateMachine* GFPSM = FindGameFeaturePluginStateMachine(PluginIdentifier))
-	{
-		const FGameFeaturePluginStateMachineProperties& MachineProperties = GFPSM->GetProperties();
-		if (MachineProperties.GetPluginProtocol() == EGameFeaturePluginProtocol::InstallBundle)
-		{
-			const FInstallBundlePluginProtocolMetaData& MetaData = MachineProperties.ProtocolMetadata.GetSubtype<FInstallBundlePluginProtocolMetaData>();
-			for (const FName& Bundle : MetaData.InstallBundles)
-			{
-				TSet<FString>& PluginsUsingBundle = DownloadedBundleToPlugin.FindOrAdd(Bundle);
-				PluginsUsingBundle.Add(FString(PluginIdentifier.GetPluginName()));
-			}
-		}
-	}
-}
-
 void UGameFeaturesSubsystem::OnGameFeatureReleasing(const FString& PluginName, const FGameFeaturePluginIdentifier& PluginIdentifier)
 {
 	CallbackObservers(EObserverCallback::Releasing, PluginIdentifier, &PluginName);
@@ -1023,20 +1004,6 @@ void UGameFeaturesSubsystem::OnGameFeaturePreMounting(const FString& PluginName,
 
 void UGameFeaturesSubsystem::OnGameFeaturePostMounting(const FString& PluginName, const FGameFeaturePluginIdentifier& PluginIdentifier, FGameFeaturePostMountingContext& Context)
 {
-	if (UGameFeaturePluginStateMachine* GFPSM = FindGameFeaturePluginStateMachine(PluginIdentifier))
-	{
-		const FGameFeaturePluginStateMachineProperties& MachineProperties = GFPSM->GetProperties();
-		if (MachineProperties.GetPluginProtocol() == EGameFeaturePluginProtocol::InstallBundle)
-		{
-			const FInstallBundlePluginProtocolMetaData& MetaData = MachineProperties.ProtocolMetadata.GetSubtype<FInstallBundlePluginProtocolMetaData>();
-			for (const FName& Bundle : MetaData.InstallBundles)
-			{
-				TSet<FString>& PluginsUsingBundle = MountedBundleToPlugin.FindOrAdd(Bundle);
-				PluginsUsingBundle.Add(PluginName);
-			}
-		}
-	}
-
 	CallbackObservers(EObserverCallback::PostMounting, PluginIdentifier, &PluginName, /*GameFeatureData=*/nullptr, &Context);
 }
 
@@ -2000,7 +1967,6 @@ void UGameFeaturesSubsystem::LoadBuiltInGameFeaturePlugin(const TSharedRef<IPlug
 				{
 					FInstallBundlePluginProtocolOptions InstallBundleOptions;
 					InstallBundleOptions.bAllowIniLoading = true;
-					InstallBundleOptions.ReleaseInstallBundleFlags = EInstallBundleReleaseRequestFlags::ExplicitRemoveList;
 					ProtocolOptions = FGameFeatureProtocolOptions(InstallBundleOptions);
 				}
 				ProtocolOptions.bForceSyncLoading = BehaviorOptions.bForceSyncLoading;
@@ -2511,47 +2477,6 @@ void UGameFeaturesSubsystem::PruneCachedGameFeaturePluginDetails(const FString&
 	CachedPluginDetailsByFilename.Remove(PluginDescriptorFilename);
 }
 
-TValueOrError<FInstallBundleReleaseRequestInfo, EInstallBundleResult> UGameFeaturesSubsystem::ReleaseBundle(const FString& PluginName, TArrayView<const FName> BundleNames, EInstallBundleReleaseRequestFlags Flags)
-{
-	TArray<FName> BundlesToRelease;
-	for (const FName& Bundle : BundleNames)
-	{
-		TSet<FString>& PluginsUsingBundle = DownloadedBundleToPlugin.FindOrAdd(Bundle);
-		PluginsUsingBundle.Remove(PluginName);
-		if (PluginsUsingBundle.Num() == 0)
-		{
-			UE_LOG(LogGameFeatures, Display, TEXT("Releasing bundle (%s) as final plugin released it (%s)"), *Bundle.ToString(), *PluginName);
-			BundlesToRelease.Add(Bundle);
-		}
-
-#if !UE_BUILD_SHIPPING
-		TSet<FString>& PluginsWithMountedBundle = MountedBundleToPlugin.FindOrAdd(Bundle);
-		ensureMsgf(!PluginsWithMountedBundle.Contains(PluginName), TEXT("Bundle (%s) is being released but still has GFPs with it mounted"), *Bundle.ToString());
-#endif
-	}
-	TSharedPtr<IInstallBundleManager> BundleManager = IInstallBundleManager::GetPlatformInstallBundleManager();
-	check(BundleManager.IsValid());
-	return BundleManager->RequestReleaseContent(BundlesToRelease, Flags);
-}
-
-TValueOrError<FInstallBundleReleaseRequestInfo, EInstallBundleResult> UGameFeaturesSubsystem::UnmountBundle(const FString& PluginName, TArrayView<const FName> BundleNames, EInstallBundleReleaseRequestFlags Flags)
-{
-	TArray<FName> BundlesToRelease;
-	for (const FName& Bundle : BundleNames)
-	{
-		TSet<FString>& PluginsUsingBundle = MountedBundleToPlugin.FindOrAdd(Bundle);
-		PluginsUsingBundle.Remove(PluginName);
-		if (PluginsUsingBundle.Num() == 0)
-		{
-			UE_LOG(LogGameFeatures, Display, TEXT("Unmounting bundle (%s) as final plugin released it (%s)"), *Bundle.ToString(), *PluginName);
-			BundlesToRelease.Add(Bundle);
-		}
-	}
-	TSharedPtr<IInstallBundleManager> BundleManager = IInstallBundleManager::GetPlatformInstallBundleManager();
-	check(BundleManager.IsValid());
-	return BundleManager->RequestReleaseContent(BundlesToRelease, Flags);
-}
-
 struct FGameFeaturePluginPredownloadContext : public FGameFeaturePluginPredownloadHandle
 {
 	const FStringView PredownloadErrorNamespace = TEXTVIEW("GameFeaturePlugin.Predownload.");
diff --git a/Engine/Plugins/Runtime/GameFeatures/Source/GameFeatures/Public/GameFeaturesProjectPolicies.h b/Engine/Plugins/Runtime/GameFeatures/Source/GameFeatures/Public/GameFeaturesProjectPolicies.h
--- a/Engine/Plugins/Runtime/GameFeatures/Source/GameFeatures/Public/GameFeaturesProjectPolicies.h
+++ b/Engine/Plugins/Runtime/GameFeatures/Source/GameFeatures/Public/GameFeaturesProjectPolicies.h
@@ -14,6 +14,12 @@ struct FPluginDependencyDetails
 	bool bFailIfNotFound = false;
 };
 
+enum class EStreamingAssetInstallMode : uint8
+{
+	GfpRequiredOnly, // only stream in data required for the GFP to load
+	Full, // stream in all data
+};
+
 // This class allows project-specific rules to be implemented for game feature plugins.
 // Create a subclass and choose it in Project Settings .. Game Features
 UCLASS()
@@ -67,8 +73,13 @@ public:
 	virtual TValueOrError<FString, FString> ResolvePluginDependency(const FString& PluginURL, const FString& DependencyName) const;
 
 	// Called to resolve install bundles for streaming asset dependencies
+	UE_DEPRECATED(5.6, "Use GetStreamingAssetInstallModes instead of creating new bundles for streaming assets.")
 	virtual TValueOrError<TArray<FName>, FString> GetStreamingAssetInstallBundles(FStringView PluginURL) const { return MakeValue(); }
 
+	// Called to resolve install modes for streaming asset dependencies
+	// Return a streaming asset install mode for each install bundle
+	virtual TValueOrError<TArray<EStreamingAssetInstallMode>, FString> GetStreamingAssetInstallModes(FStringView PluginURL, TConstArrayView<FName> InstallBundleNames) const;
+
 	// Called by code that explicitly wants to load a specific plugin
 	// (e.g., when using a fast cook a game might want to allow explicitly loaded game feature plugins)
 	virtual void ExplicitLoadGameFeaturePlugin(const FString& PluginURL, const FGameFeaturePluginLoadComplete& CompleteDelegate, const bool bActivateGameFeatures);
diff --git a/Engine/Plugins/Runtime/GameFeatures/Source/GameFeatures/Public/GameFeaturesSubsystem.h b/Engine/Plugins/Runtime/GameFeatures/Source/GameFeatures/Public/GameFeaturesSubsystem.h
--- a/Engine/Plugins/Runtime/GameFeatures/Source/GameFeatures/Public/GameFeaturesSubsystem.h
+++ b/Engine/Plugins/Runtime/GameFeatures/Source/GameFeatures/Public/GameFeaturesSubsystem.h
@@ -341,10 +341,16 @@ struct GAMEFEATURES_API FInstallBundlePluginProtocolOptions
 
 	FInstallBundlePluginProtocolOptions();
 
+	PRAGMA_DISABLE_DEPRECATION_WARNINGS
+	FInstallBundlePluginProtocolOptions(const FInstallBundlePluginProtocolOptions&) = default;
+	FInstallBundlePluginProtocolOptions& operator=(const FInstallBundlePluginProtocolOptions&) = default;
+	PRAGMA_ENABLE_DEPRECATION_WARNINGS
+
 	/** EInstallBundleRequestFlags utilized during the download/install by InstallBundleManager */
 	EInstallBundleRequestFlags InstallBundleFlags;
 
 	/** EInstallBundleReleaseRequestFlags utilized during our release and uninstall states */
+	UE_DEPRECATED(5.6, "Release flags are now applied internally and no longer need to be explicitly set.")
 	EInstallBundleReleaseRequestFlags ReleaseInstallBundleFlags;
 
 	/** If we want to attempt to uninstall InstallBundle data installed by this plugin before terminating */
@@ -669,7 +675,6 @@ private:
 	void OnGameFeaturePredownloading(const FString& PluginName, const FGameFeaturePluginIdentifier& PluginIdentifier);
 
 	void OnGameFeatureDownloading(const FString& PluginName, const FGameFeaturePluginIdentifier& PluginIdentifier);
-	void OnGameFeatureDownloaded(const FGameFeaturePluginIdentifier& PluginIdentifier);
 	friend struct FGameFeaturePluginState_Downloading;
 
 	void OnGameFeatureReleasing(const FString& PluginName, const FGameFeaturePluginIdentifier& PluginIdentifier);
@@ -731,22 +736,6 @@ private:
 	friend struct FGameFeaturePluginState_Unmounting;
 	friend struct FBaseDataReleaseGameFeaturePluginState;
 
-	/**
-	 * Wrapper to InstallBundleManager release bundle.
-	 * Will only release the bundle if this is the last GFP using the bundle.
-	 */
-	TValueOrError<FInstallBundleReleaseRequestInfo, EInstallBundleResult> ReleaseBundle(const FString& PluginName, TArrayView<const FName> BundleNames, EInstallBundleReleaseRequestFlags Flags);
-
-	/**
-	 * Wrapper to InstallBundleManager release bundle for the unmount call
-	 * Will only unmount the bundle if this is the last GFP using the bundle.
-	 */
-	TValueOrError<FInstallBundleReleaseRequestInfo, EInstallBundleResult> UnmountBundle(const FString& PluginName,  TArrayView<const FName> BundleNames, EInstallBundleReleaseRequestFlags Flags);
-
-	// Reference tracker to know which bundles are in use.
-	TMap<FName, TSet<FString>> MountedBundleToPlugin;
-	TMap<FName, TSet<FString>> DownloadedBundleToPlugin;
-
 	/** Gets the state machine associated with the specified URL */
 	UGameFeaturePluginStateMachine* FindGameFeaturePluginStateMachine(const FString& PluginURL) const;
 
