Clear array after uninstalling content

diff --git a/Engine/Plugins/Runtime/GameFeatures/Source/GameFeatures/Private/GameFeatureOptionalContentInstaller.cpp b/Engine/Plugins/Runtime/GameFeatures/Source/GameFeatures/Private/GameFeatureOptionalContentInstaller.cpp
--- a/Engine/Plugins/Runtime/GameFeatures/Source/GameFeatures/Private/GameFeatureOptionalContentInstaller.cpp
+++ b/Engine/Plugins/Runtime/GameFeatures/Source/GameFeatures/Private/GameFeatureOptionalContentInstaller.cpp
@@ -67,6 +67,7 @@ void UGameFeatureOptionalContentInstaller::UninstallContent()
 		UE_LOG(LogGameFeatures, Log, TEXT("Uninstalling Optional bundles for %s"), *GFP);
 		ReleaseContent(GFP, EInstallBundleReleaseRequestFlags::RemoveFilesIfPossible);
 	}
+	RelevantGFPs.Empty();
 }
 
 void UGameFeatureOptionalContentInstaller::EnableCellularDownloading(bool bEnable)
