Adding post predownload callback observer.
This allows systems to hook into this callback and remove module info after it has been predownload to not lead to a forever growing list of bundles.

diff --git a/Engine/Plugins/Runtime/GameFeatures/Source/GameFeatures/Private/GameFeaturesSubsystem.cpp b/Engine/Plugins/Runtime/GameFeatures/Source/GameFeatures/Private/GameFeaturesSubsystem.cpp
--- a/Engine/Plugins/Runtime/GameFeatures/Source/GameFeatures/Private/GameFeaturesSubsystem.cpp
+++ b/Engine/Plugins/Runtime/GameFeatures/Source/GameFeatures/Private/GameFeaturesSubsystem.cpp
@@ -987,6 +987,11 @@ void UGameFeaturesSubsystem::OnGameFeaturePredownloading(const FString& PluginNa
 	CallbackObservers(EObserverCallback::Predownloading, PluginIdentifier, &PluginName);
 }
 
+void UGameFeaturesSubsystem::OnGameFeaturePostPredownloading(const FString& PluginName, const FGameFeaturePluginIdentifier& PluginIdentifier)
+{
+	CallbackObservers(EObserverCallback::PostPredownloading, PluginIdentifier, &PluginName);
+}
+
 void UGameFeaturesSubsystem::OnGameFeatureDownloading(const FString& PluginName, const FGameFeaturePluginIdentifier& PluginIdentifier)
 {
 	CallbackObservers(EObserverCallback::Downloading, PluginIdentifier, &PluginName);
@@ -2784,6 +2789,10 @@ struct FGameFeaturePluginPredownloadContext : public FGameFeaturePluginPredownlo
 		}
 
 		BundleManager->RequestReleaseContent(ReleaseList, EInstallBundleReleaseRequestFlags::None, KeepList);
+		for (const TPair<FGameFeaturePluginIdentifier, FInstallBundlePluginProtocolMetaData>& Pair : GFPs)
+		{
+			GFPSubSys.OnGameFeaturePostPredownloading(FString(Pair.Key.GetPluginName()), Pair.Key);
+		}
 	}
 };
 
@@ -3204,7 +3213,7 @@ void UGameFeaturesSubsystem::CallbackObservers(EObserverCallback CallbackType, c
 	const UGameFeatureData* GameFeatureData /*= nullptr*/, 
 	FGameFeatureStateChangeContext* StateChangeContext /*= nullptr*/)
 {
-	static_assert(std::underlying_type<EObserverCallback>::type(EObserverCallback::Count) == 15, "Update UGameFeaturesSubsystem::CallbackObservers to handle added EObserverCallback");
+	static_assert(std::underlying_type<EObserverCallback>::type(EObserverCallback::Count) == 16, "Update UGameFeaturesSubsystem::CallbackObservers to handle added EObserverCallback");
 
 	// Protect against modifying the observer list during iteration
 	TArray<UObject*> LocalObservers(Observers);
@@ -3236,6 +3245,15 @@ void UGameFeaturesSubsystem::CallbackObservers(EObserverCallback CallbackType, c
 		}
 		break;
 	}
+	case EObserverCallback::PostPredownloading:
+	{
+		check(PluginName);
+		for (UObject* Observer : LocalObservers)
+		{
+			CastChecked<IGameFeatureStateChangeObserver>(Observer)->OnGameFeaturePostPredownloading(*PluginName, PluginIdentifier);
+		}
+		break;
+	}
 	case EObserverCallback::Downloading:
 	{
 		check(PluginName);
diff --git a/Engine/Plugins/Runtime/GameFeatures/Source/GameFeatures/Public/GameFeatureStateChangeObserver.h b/Engine/Plugins/Runtime/GameFeatures/Source/GameFeatures/Public/GameFeatureStateChangeObserver.h
--- a/Engine/Plugins/Runtime/GameFeatures/Source/GameFeatures/Public/GameFeatureStateChangeObserver.h
+++ b/Engine/Plugins/Runtime/GameFeatures/Source/GameFeatures/Public/GameFeatureStateChangeObserver.h
@@ -43,6 +43,9 @@ public:
 	// Invoked when content begins installing via predownload
 	virtual void OnGameFeaturePredownloading(const FString& PluginName, const FGameFeaturePluginIdentifier& PluginIdentifier) {}
 
+	// Invoked when pre-downloading content has finished.
+	virtual void OnGameFeaturePostPredownloading(const FString& PluginName, const FGameFeaturePluginIdentifier& PluginIdentifier) {}
+
 	// Invoked when content begins installing
 	virtual void OnGameFeatureDownloading(const FString& PluginName, const FGameFeaturePluginIdentifier& PluginIdentifier) {}
 
diff --git a/Engine/Plugins/Runtime/GameFeatures/Source/GameFeatures/Public/GameFeaturesSubsystem.h b/Engine/Plugins/Runtime/GameFeatures/Source/GameFeatures/Public/GameFeaturesSubsystem.h
--- a/Engine/Plugins/Runtime/GameFeatures/Source/GameFeatures/Public/GameFeaturesSubsystem.h
+++ b/Engine/Plugins/Runtime/GameFeatures/Source/GameFeatures/Public/GameFeaturesSubsystem.h
@@ -673,6 +673,7 @@ private:
 	friend struct FGameFeaturePluginState_CheckingStatus;
 
 	void OnGameFeaturePredownloading(const FString& PluginName, const FGameFeaturePluginIdentifier& PluginIdentifier);
+	void OnGameFeaturePostPredownloading(const FString& PluginName, const FGameFeaturePluginIdentifier& PluginIdentifier);
 
 	void OnGameFeatureDownloading(const FString& PluginName, const FGameFeaturePluginIdentifier& PluginIdentifier);
 	friend struct FGameFeaturePluginState_Downloading;
@@ -786,6 +787,7 @@ private:
 		CheckingStatus,
 		Terminating,
 		Predownloading,
+		PostPredownloading,
 		Downloading,
 		Releasing,
 		PreMounting,
