Only send OnGameFeaturePostPredownloading if a GFP machine is not active.

diff --git a/Engine/Plugins/Runtime/GameFeatures/Source/GameFeatures/Private/GameFeaturesSubsystem.cpp b/Engine/Plugins/Runtime/GameFeatures/Source/GameFeatures/Private/GameFeaturesSubsystem.cpp
--- a/Engine/Plugins/Runtime/GameFeatures/Source/GameFeatures/Private/GameFeaturesSubsystem.cpp
+++ b/Engine/Plugins/Runtime/GameFeatures/Source/GameFeatures/Private/GameFeaturesSubsystem.cpp
@@ -2785,14 +2785,13 @@ struct FGameFeaturePluginPredownloadContext : public FGameFeaturePluginPredownlo
 				continue;
 			}
 
+			// Only send this if a GFP is not using the bundle
+			GFPSubSys.OnGameFeaturePostPredownloading(FString(Pair.Key.GetPluginName()), Pair.Key);
+
 			ReleaseList.Append(Pair.Value.InstallBundles);
 		}
 
 		BundleManager->RequestReleaseContent(ReleaseList, EInstallBundleReleaseRequestFlags::None, KeepList);
-		for (const TPair<FGameFeaturePluginIdentifier, FInstallBundlePluginProtocolMetaData>& Pair : GFPs)
-		{
-			GFPSubSys.OnGameFeaturePostPredownloading(FString(Pair.Key.GetPluginName()), Pair.Key);
-		}
 	}
 };
 
