UE: Make InitializeBasePluginIniFile static since it does not use this and it is currently being called on a pointer that is always null at the time it is called

diff --git a/Engine/Plugins/Runtime/GameFeatures/Source/GameFeatures/Private/GameFeatureData.cpp b/Engine/Plugins/Runtime/GameFeatures/Source/GameFeatures/Private/GameFeatureData.cpp
--- a/Engine/Plugins/Runtime/GameFeatures/Source/GameFeatures/Private/GameFeatureData.cpp
+++ b/Engine/Plugins/Runtime/GameFeatures/Source/GameFeatures/Private/GameFeatureData.cpp
@@ -82,7 +82,7 @@ static TAutoConsoleVariable<bool> CVarAllowRuntimeDeviceProfiles(
 	TEXT("Allow game feature plugins to generate device profiles from config based on existing parents"),
 	ECVF_Default);
 
-void UGameFeatureData::InitializeBasePluginIniFile(const FString& PluginInstalledFilename) const
+void UGameFeatureData::InitializeBasePluginIniFile(const FString& PluginInstalledFilename)
 {
 	const FString PluginName = FPaths::GetBaseFilename(PluginInstalledFilename);
 
@@ -107,7 +107,7 @@ void UGameFeatureData::InitializeBasePluginIniFile(const FString& PluginInstalle
 		if (FConfigCacheIni::LoadExternalIniFile(PluginConfig, *PluginName, *EngineConfigDir, *PluginConfigDir, bIsBaseIniName, nullptr, bForceReloadFromDisk, bWriteDestIni))
 		{
 			// This is the deprecated loading path that doesn't handle cases like + in arrays
-			UE_LOG(LogGameFeatures, Log, TEXT("[GameFeatureData %s]: Loaded deprecated config %s, rename to start with Default for normal parsing"), *GetPathNameSafe(this), *PluginConfigFilename);
+			UE_LOG(LogGameFeatures, Log, TEXT("Loaded deprecated config %s, rename to start with Default for normal parsing"), *PluginConfigFilename);
 
 			// register this plugin, so the ConfigContext.Load, and future loads, know about it
 			if (bUseNewDynamicLayers)
@@ -571,7 +571,7 @@ void UGameFeatureData::InitializeHierarchicalPluginIniFiles(const FString& Plugi
 	}
 }
 
-void UGameFeatureData::ReloadConfigs(FConfigFile& PluginConfig) const
+void UGameFeatureData::ReloadConfigs(FConfigFile& PluginConfig)
 {
 	// Reload configs so objects get the changes
 	for (const auto& ConfigEntry : AsConst(PluginConfig))
@@ -609,7 +609,7 @@ void UGameFeatureData::ReloadConfigs(FConfigFile& PluginConfig) const
 			}
 			else
 			{
-				UE_LOG(LogGameFeatures, Warning, TEXT("[GameFeatureData %s]: Couldn't find PerObjectConfig class %s for %s while processing %s, config changes won't be reloaded."), *GetPathNameSafe(this), *ClassName, *ObjectName, *PluginConfig.Name.ToString());
+				UE_LOG(LogGameFeatures, Warning, TEXT("Couldn't find PerObjectConfig class %s for %s while processing %s, config changes won't be reloaded."), *ClassName, *ObjectName, *PluginConfig.Name.ToString());
 			}
 		}
 		// Standard INI section case
diff --git a/Engine/Plugins/Runtime/GameFeatures/Source/GameFeatures/Private/GameFeaturePluginStateMachine.cpp b/Engine/Plugins/Runtime/GameFeatures/Source/GameFeatures/Private/GameFeaturePluginStateMachine.cpp
--- a/Engine/Plugins/Runtime/GameFeatures/Source/GameFeatures/Private/GameFeaturePluginStateMachine.cpp
+++ b/Engine/Plugins/Runtime/GameFeatures/Source/GameFeatures/Private/GameFeaturePluginStateMachine.cpp
@@ -2511,7 +2511,7 @@ struct FGameFeaturePluginState_Mounting : public FGameFeaturePluginState
 		{
 			if (AllowIniLoading())
 			{
-				StateProperties.GameFeatureData->InitializeBasePluginIniFile(StateProperties.PluginInstalledFilename);
+				UGameFeatureData::InitializeBasePluginIniFile(StateProperties.PluginInstalledFilename);
 			}
 
 			FGameFeaturePostMountingContext Context(StateProperties.PluginName, [this](FStringView InPauserTag) { OnPostMountPauserCompleted(InPauserTag); });
diff --git a/Engine/Plugins/Runtime/GameFeatures/Source/GameFeatures/Public/GameFeatureData.h b/Engine/Plugins/Runtime/GameFeatures/Source/GameFeatures/Public/GameFeatureData.h
--- a/Engine/Plugins/Runtime/GameFeatures/Source/GameFeatures/Public/GameFeatureData.h
+++ b/Engine/Plugins/Runtime/GameFeatures/Source/GameFeatures/Public/GameFeatureData.h
@@ -39,7 +39,7 @@ public:
 #endif //if WITH_EDITOR
 
 	/** Method to process the base ini file for the plugin during loading */
-	void InitializeBasePluginIniFile(const FString& PluginInstalledFilename) const;
+	static void InitializeBasePluginIniFile(const FString& PluginInstalledFilename);
 
 	/** Method to process ini files for the plugin during activation */
 	void InitializeHierarchicalPluginIniFiles(const FString& PluginInstalledFilename) const;
@@ -98,7 +98,7 @@ private:
 #endif
 
 	/** Internal helper function to reload config data on objects as a result of a plugin INI being loaded */
-	void ReloadConfigs(FConfigFile& PluginConfig) const;
+	static void ReloadConfigs(FConfigFile& PluginConfig);
 
 protected:
 
