Call ReleaseBundlesIfPossible before any pre-download compete or cancel. This prevents modules hanging around in bundle manager or bundle sources when content is already up to date.

diff --git a/Engine/Plugins/Runtime/GameFeatures/Source/GameFeatures/Private/GameFeaturesSubsystem.cpp b/Engine/Plugins/Runtime/GameFeatures/Source/GameFeatures/Private/GameFeaturesSubsystem.cpp
--- a/Engine/Plugins/Runtime/GameFeatures/Source/GameFeatures/Private/GameFeaturesSubsystem.cpp
+++ b/Engine/Plugins/Runtime/GameFeatures/Source/GameFeatures/Private/GameFeaturesSubsystem.cpp
@@ -2538,6 +2538,8 @@ struct FGameFeaturePluginPredownloadContext : public FGameFeaturePluginPredownlo
 
 	void SetComplete()
 	{
+		ReleaseBundlesIfPossible();
+
 		bIsComplete = true;
 		if (OnComplete)
 		{
@@ -2547,6 +2549,8 @@ struct FGameFeaturePluginPredownloadContext : public FGameFeaturePluginPredownlo
 
 	void SetComplete(UE::GameFeatures::FResult&& InResult)
 	{
+		ReleaseBundlesIfPossible();
+
 		Result = MoveTemp(InResult);
 		bIsComplete = true;
 		if (OnComplete)
@@ -2557,6 +2561,8 @@ struct FGameFeaturePluginPredownloadContext : public FGameFeaturePluginPredownlo
 
 	void SetCompleteCanceled()
 	{
+		ReleaseBundlesIfPossible();
+
 		Result = MakeError(FString::Printf(TEXT("%.*s%s"),
 			PredownloadErrorNamespace.Len(), PredownloadErrorNamespace.GetData(),
 			TEXT("Canceled")));
@@ -2754,8 +2760,6 @@ struct FGameFeaturePluginPredownloadContext : public FGameFeaturePluginPredownlo
 		FTSTicker::GetCoreTicker().AddTicker(
 			FTickerDelegate::CreateLambda([Context = SharedThis(this)](float)
 			{
-				Context->ReleaseBundlesIfPossible();
-
 				// Done
 				Context->SetComplete();
 
