First pass at GFP IAD API.
- Fix some IAD progress issues.

diff --git a/Engine/Plugins/Runtime/GameFeatures/Source/GameFeatures/GameFeatures.Build.cs b/Engine/Plugins/Runtime/GameFeatures/Source/GameFeatures/GameFeatures.Build.cs
--- a/Engine/Plugins/Runtime/GameFeatures/Source/GameFeatures/GameFeatures.Build.cs
+++ b/Engine/Plugins/Runtime/GameFeatures/Source/GameFeatures/GameFeatures.Build.cs
@@ -34,6 +34,13 @@ namespace UnrealBuildTool.Rules
 				}
 			);
 
+			PublicIncludePathModuleNames.AddRange(
+				new string[]
+				{
+					"IoStoreOnDemand",
+				}
+			);
+
 			if (Target.bBuildEditor)
 			{
 				PrivateDependencyModuleNames.AddRange(
diff --git a/Engine/Plugins/Runtime/GameFeatures/Source/GameFeatures/Private/GameFeaturePluginStateMachine.cpp b/Engine/Plugins/Runtime/GameFeatures/Source/GameFeatures/Private/GameFeaturePluginStateMachine.cpp
--- a/Engine/Plugins/Runtime/GameFeatures/Source/GameFeatures/Private/GameFeaturePluginStateMachine.cpp
+++ b/Engine/Plugins/Runtime/GameFeatures/Source/GameFeatures/Private/GameFeaturePluginStateMachine.cpp
@@ -2784,9 +2784,7 @@ struct FGameFeaturePluginState_AssetDependencyStreaming : public FGameFeaturePlu
 						return;
 					}
 
-					const double RequestProgress = (Progress.TotalInstallSize > 0) ?
-						FMath::Clamp(double(Progress.CurrentInstallSize) / Progress.TotalInstallSize, 0.0, 1.0) :
-						0;
+					const double RequestProgress = Progress.GetRelativeProgress();
 
 					FIoStoreOnDemandProgress* MyProgress = Algo::FindBy(
 						LambdaOnDemandContext->Progress, InstallBundle, &FIoStoreOnDemandProgress::InstallBundle);
diff --git a/Engine/Plugins/Runtime/GameFeatures/Source/GameFeatures/Private/GameFeaturesSubsystem.cpp b/Engine/Plugins/Runtime/GameFeatures/Source/GameFeatures/Private/GameFeaturesSubsystem.cpp
--- a/Engine/Plugins/Runtime/GameFeatures/Source/GameFeatures/Private/GameFeaturesSubsystem.cpp
+++ b/Engine/Plugins/Runtime/GameFeatures/Source/GameFeatures/Private/GameFeaturesSubsystem.cpp
@@ -2809,6 +2809,230 @@ TSharedRef<FGameFeaturePluginPredownloadHandle> UGameFeaturesSubsystem::Predownl
 	return Context;
 }
 
+FGameFeaturePluginStreamedAssetsRef::FGameFeaturePluginStreamedAssetsRef(FStringView DebugName)
+	: ContentHandle(UE::IoStore::FOnDemandContentHandle::Create(DebugName))
+{
+}
+
+FString FGameFeaturePluginStreamedAssetsRef::GetDebugName() const
+{
+	return LexToString(ContentHandle);
+}
+
+void FGameFeaturePluginStreamedAssetsRef::Release()
+{
+	ContentHandle.Reset();
+}
+
+bool FGameFeaturePluginStreamedAssetsRef::IsValid() const
+{
+	return ContentHandle.IsValid();
+}
+
+struct FGameFeatureAssetStreamContext : public FGameFeaturePluginAssetStreamHandle
+{
+	const FStringView ErrorNamespace = TEXTVIEW("GameFeaturePlugin.AssetStream.");
+	FGameFeaturePluginStreamedAssetsRef AssetRef;
+	TUniqueFunction<void(const UE::GameFeatures::FResult&)> OnComplete;
+	TUniqueFunction<void(float)> OnProgress;
+	UE::GameFeatures::FResult Result = MakeValue();
+	float Progress = 0.0f;
+	bool bIsComplete = false;
+	UE::IoStore::FOnDemandCancellationToken CancellationToken;
+
+	FGameFeatureAssetStreamContext(FGameFeaturePluginStreamedAssetsRef&& InAssetRef)
+		: AssetRef(MoveTemp(InAssetRef))
+	{}
+
+	virtual bool IsComplete() const override
+	{
+		return bIsComplete;
+	}
+
+	virtual const UE::GameFeatures::FResult& GetResult() const override
+	{
+		return Result;
+	}
+
+	virtual float GetProgress() const override
+	{
+		return Progress;
+	}
+
+	virtual void Cancel() override
+	{
+		CancellationToken.Cancel();
+	}
+
+	virtual FGameFeaturePluginStreamedAssetsRef GetAssetsRef() const override
+	{
+		return AssetRef;
+	}
+
+	void SetComplete()
+	{
+		bIsComplete = true;
+		if (OnComplete)
+		{
+			OnComplete(Result);
+		}
+	}
+
+	void SetComplete(UE::GameFeatures::FResult&& InResult)
+	{
+		Result = MoveTemp(InResult);
+		bIsComplete = true;
+		if (OnComplete)
+		{
+			OnComplete(Result);
+		}
+	}
+
+	void SetCompleteCanceled()
+	{
+		Result = MakeError(FString::Printf(TEXT("%.*s%s"),
+			ErrorNamespace.Len(), ErrorNamespace.GetData(),
+			TEXT("Canceled")));
+		bIsComplete = true;
+		if (OnComplete)
+		{
+			OnComplete(Result);
+		}
+	}
+
+	void Start(TArrayView<FName> AssetPackageNames)
+	{
+		if (CancellationToken.IsCanceled())
+		{
+			SetCompleteCanceled();
+			return;
+		}
+
+		UGameFeaturesSubsystem& GFPSubSys = UGameFeaturesSubsystem::Get();
+
+		TSet<FStringView> MountPoints;
+
+		// Find mount points
+		FNameBuilder PackageNameBuilder;
+		for (FName PackageName : AssetPackageNames)
+		{
+			PackageName.ToString(PackageNameBuilder);
+			const FStringView MountPoint = FPathViews::GetMountPointNameFromPath(PackageNameBuilder);
+			MountPoints.Add(MountPoint);
+		}
+
+		// Verify that the we can stream from these GFPs
+		for (FStringView MountPoint : MountPoints)
+		{
+			FString PluginURL;
+			if (!GFPSubSys.GetPluginURLByName(MountPoint, PluginURL))
+			{
+				UE_LOGFMT(LogGameFeatures, Error, "{MountPoint} is unknown", ("MountPoint", MountPoint));
+				UE::GameFeatures::FResult ErrorResult = MakeError(FString::Printf(TEXT("%.*s%s"),
+					ErrorNamespace.Len(), ErrorNamespace.GetData(), TEXT("Unknown_Mount_Point")));
+				SetComplete(MoveTemp(ErrorResult));
+				return;
+			}
+
+			const UGameFeaturePluginStateMachine* Machine = GFPSubSys.FindGameFeaturePluginStateMachine(PluginURL);
+			if (!Machine)
+			{
+				UE_LOGFMT(LogGameFeatures, Error, "Could not find a GFP with mount point {MountPoint}", ("MountPoint", MountPoint));
+				UE::GameFeatures::FResult ErrorResult = MakeError(FString::Printf(TEXT("%.*s%s"),
+					ErrorNamespace.Len(), ErrorNamespace.GetData(), TEXT("Unknown_GFP_Mount_Point")));
+				SetComplete(MoveTemp(ErrorResult));
+				return;
+			}
+
+			// Check installed and not moving to a state less than installed
+			if (Machine->GetCurrentState() >= EGameFeaturePluginState::Installed &&
+				Machine->GetDestination() > EGameFeaturePluginState::StatusKnown)
+			{
+				UE_LOGFMT(LogGameFeatures, Error, "Invalid GFP state for {MountPoint}", ("MountPoint", MountPoint));
+				UE::GameFeatures::FResult ErrorResult = MakeError(FString::Printf(TEXT("%.*s%s"),
+					ErrorNamespace.Len(), ErrorNamespace.GetData(), TEXT("GFP_Not_Ready")));
+				SetComplete(MoveTemp(ErrorResult));
+				return;
+			}
+		}
+
+		// Install the packages
+		UE::IoStore::FIoStoreOnDemandModule* IoStoreOnDemandModule = FModuleManager::Get().GetModulePtr<UE::IoStore::FIoStoreOnDemandModule>(TEXT("IoStoreOnDemand"));
+		if (!IoStoreOnDemandModule)
+		{
+			UE::GameFeatures::FResult ErrorResult = MakeError(FString::Printf(TEXT("%.*s%s"),
+				ErrorNamespace.Len(), ErrorNamespace.GetData(), TEXT("IoStoreOnDemand.ModuleNotFound")));
+			SetComplete(MoveTemp(ErrorResult));
+			return;
+		}
+
+		if (!AssetRef.IsValid())
+		{
+			UE::GameFeatures::FResult ErrorResult = MakeError(FString::Printf(TEXT("%.*s%s"),
+				ErrorNamespace.Len(), ErrorNamespace.GetData(), TEXT("Invalid_AssetRef")));
+			SetComplete(MoveTemp(ErrorResult));
+			return;
+		}
+
+		UE::IoStore::FOnDemandInstallArgs InstallArgs;
+		InstallArgs.Options |= UE::IoStore::EOnDemandInstallOptions::CallbackOnGameThread;
+		InstallArgs.ContentHandle = AssetRef.ContentHandle;
+
+		InstallArgs.PackageIds.Reserve(AssetPackageNames.Num());
+		Algo::Transform(AssetPackageNames, InstallArgs.PackageIds, &FPackageId::FromName);
+
+		IoStoreOnDemandModule->Install(MoveTemp(InstallArgs),
+			[Context = SharedThis(this)](const UE::IoStore::FOnDemandInstallResult& OnDemandInstallResult)
+			{
+				Context->OnInstallComplete(OnDemandInstallResult);
+			},
+			[Context = SharedThis(this)](const UE::IoStore::FOnDemandInstallProgress& InProgress)
+			{
+				Context->OnInstallProgress(InProgress);
+			},
+			&CancellationToken);
+	}
+
+	void OnInstallComplete(const UE::IoStore::FOnDemandInstallResult& OnDemandInstallResult)
+	{
+		if (!OnDemandInstallResult.Status.IsOk())
+		{
+			const FString StatusString = OnDemandInstallResult.Status.ToString();
+			FString ErrorCode = GetIoErrorText(OnDemandInstallResult.Status.GetErrorCode());
+			ErrorCode.ReplaceCharInline(TEXT(' '), TEXT('_'), ESearchCase::CaseSensitive);
+
+			UE::GameFeatures::FResult ErrorResult(MakeError(
+				FString::Printf(TEXT("%.*s%s.%s"),
+					ErrorNamespace.Len(), ErrorNamespace.GetData(), TEXT("IoStoreOnDemand.OnComplete"), *ErrorCode)),
+				FText::AsCultureInvariant(StatusString));
+			SetComplete(MoveTemp(ErrorResult));
+			return;
+		}
+
+		SetComplete();
+	}
+
+	void OnInstallProgress(const UE::IoStore::FOnDemandInstallProgress& InProgress)
+	{
+		Progress = InProgress.GetRelativeProgress();
+		if (OnProgress)
+		{
+			OnProgress(Progress);
+		}
+	}
+};
+
+TSharedRef<FGameFeaturePluginAssetStreamHandle> UGameFeaturesSubsystem::DownloadStreamedAssets(FGameFeaturePluginStreamedAssetsRef AssetRef, TArrayView<FName> AssetPackageNames,
+	TUniqueFunction<void(const UE::GameFeatures::FResult&)> OnComplete /*= nullptr*/, TUniqueFunction<void(float)> OnProgress /*= nullptr*/)
+{
+	TSharedRef<FGameFeatureAssetStreamContext> Context = MakeShared<FGameFeatureAssetStreamContext>(MoveTemp(AssetRef));
+	Context->OnComplete = MoveTemp(OnComplete);
+	Context->OnProgress = MoveTemp(OnProgress);
+	Context->Start(AssetPackageNames);
+
+	return Context;
+}
+
 UGameFeaturePluginStateMachine* UGameFeaturesSubsystem::FindGameFeaturePluginStateMachine(const FString& PluginURL) const
 {
 	FGameFeaturePluginIdentifier FindPluginIdentifier(PluginURL);
diff --git a/Engine/Plugins/Runtime/GameFeatures/Source/GameFeatures/Public/GameFeaturesSubsystem.h b/Engine/Plugins/Runtime/GameFeatures/Source/GameFeatures/Public/GameFeaturesSubsystem.h
--- a/Engine/Plugins/Runtime/GameFeatures/Source/GameFeatures/Public/GameFeaturesSubsystem.h
+++ b/Engine/Plugins/Runtime/GameFeatures/Source/GameFeatures/Public/GameFeaturesSubsystem.h
@@ -5,6 +5,7 @@
 #include "Containers/Union.h"
 #include "Engine/Engine.h"
 #include "GameFeatureTypesFwd.h"
+#include "IO/IoStoreOnDemand.h"
 
 #include "GameFeaturesSubsystem.generated.h"
 
@@ -260,15 +261,42 @@ struct FBuiltInGameFeaturePluginBehaviorOptions
 	bool bLogErrorOnForcedDependencyCreation = false;
 };
 
-struct FGameFeaturePluginPredownloadHandle : public TSharedFromThis<FGameFeaturePluginPredownloadHandle>
+struct FGameFeaturePluginAsyncHandle : public TSharedFromThis<FGameFeaturePluginAsyncHandle>
 {
-	virtual ~FGameFeaturePluginPredownloadHandle() {}
+	virtual ~FGameFeaturePluginAsyncHandle() {}
 	virtual bool IsComplete() const = 0;
 	virtual const UE::GameFeatures::FResult& GetResult() const = 0;
 	virtual float GetProgress() const = 0;
 	virtual void Cancel() = 0;
 };
 
+/** Handle to track a GFP predownload */
+struct FGameFeaturePluginPredownloadHandle : public FGameFeaturePluginAsyncHandle
+{
+};
+
+/** Referencer to hold individual asset downloads in local cache */
+struct FGameFeaturePluginStreamedAssetsRef
+{
+	explicit FGameFeaturePluginStreamedAssetsRef(FStringView DebugName);
+
+	FString GetDebugName() const;
+
+	void Release();
+	bool IsValid() const;
+	operator bool() const { return IsValid(); }
+
+private:
+	UE::IoStore::FOnDemandContentHandle ContentHandle;
+	friend struct FGameFeatureAssetStreamContext;
+};
+
+/** Handle to track downloads of individual assets */
+struct FGameFeaturePluginAssetStreamHandle : public FGameFeaturePluginAsyncHandle
+{
+	virtual FGameFeaturePluginStreamedAssetsRef GetAssetsRef() const = 0;
+};
+
 /** Struct used to transform a GameFeaturePlugin URL into something that can uniquely identify the GameFeaturePlugin
     without including any transient data being passed in through the URL */
 USTRUCT()
@@ -650,9 +678,18 @@ public:
 	 * Pre-install any required game feature data, which can be useful for larger payloads. 
 	 * This does not instantiate any GFP although it is safe to do so before this finishes. 
 	*/
-	TSharedRef<FGameFeaturePluginPredownloadHandle> PredownloadGameFeaturePlugins(TConstArrayView<FString> PluginURLs, TUniqueFunction<void(const UE::GameFeatures::FResult&)> OnComplete = nullptr, TUniqueFunction<void(float)> OnProgress = nullptr);
+	TSharedRef<FGameFeaturePluginPredownloadHandle> PredownloadGameFeaturePlugins(TConstArrayView<FString> PluginURLs, 
+		TUniqueFunction<void(const UE::GameFeatures::FResult&)> OnComplete = nullptr, TUniqueFunction<void(float)> OnProgress = nullptr);
 	friend struct FGameFeaturePluginPredownloadContext;
 
+	/**
+	 * Download streamed assets from installed GFPs. The assets cannot be loaded unless the plugin is also mounted.
+	 * Assets are gauranteed to be available as long as AssetRef is valid.
+	 */
+	TSharedRef<FGameFeaturePluginAssetStreamHandle> DownloadStreamedAssets(FGameFeaturePluginStreamedAssetsRef AssetRef, TArrayView<FName> AssetPackageNames,
+		TUniqueFunction<void(const UE::GameFeatures::FResult&)> OnComplete = nullptr, TUniqueFunction<void(float)> OnProgress = nullptr);
+	friend struct FGameFeatureAssetStreamContext;
+
 	/** Determine the initial feature state for a built-in plugin */
 	static EBuiltInAutoState DetermineBuiltInInitialFeatureState(TSharedPtr<FJsonObject> Descriptor, const FString& ErrorContext);
 
diff --git a/Engine/Source/Runtime/Experimental/IoStore/OnDemand/Public/IO/IoStoreOnDemand.h b/Engine/Source/Runtime/Experimental/IoStore/OnDemand/Public/IO/IoStoreOnDemand.h
--- a/Engine/Source/Runtime/Experimental/IoStore/OnDemand/Public/IO/IoStoreOnDemand.h
+++ b/Engine/Source/Runtime/Experimental/IoStore/OnDemand/Public/IO/IoStoreOnDemand.h
@@ -354,27 +354,40 @@ struct FOnDemandInstallProgress
 	uint64 CurrentInstallSize = 0;
 
 
-	FORCEINLINE uint64 GetTotalDownloadSize() const
+	uint64 GetTotalDownloadSize() const
 	{
 		return TotalInstallSize;
 	}
 
-	FORCEINLINE uint64 GetAlreadyDownloadedSize() const
+	uint64 GetAlreadyDownloadedSize() const
 	{
 		return CurrentInstallSize;
 	}
 
-	FORCEINLINE float GetProgress() const 
+	float GetRelativeProgress() const 
 	{
-		return float(GetTotalDownloadSize()) / float(GetAlreadyDownloadedSize());
+		const double Progress = (TotalInstallSize > 0) ?
+			FMath::Clamp(double(CurrentInstallSize) / double(TotalInstallSize), 0.0, 1.0) :
+			0;
+
+		return float(Progress);
+	}
+
+	float GetAbsoluteProgress() const
+	{
+		const double Progress = (TotalContentSize > 0) ?
+			FMath::Clamp(double(GetCachedSize()) / double(TotalContentSize), 0.0, 1.0) :
+			0;
+
+		return float(Progress);
 	}
 
-	FORCEINLINE uint64 GetCachedSize() const 
+	uint64 GetCachedSize() const 
 	{
-		return TotalContentSize - TotalInstallSize;
+		return TotalContentSize - TotalInstallSize + CurrentInstallSize;
 	}
 
-	FORCEINLINE uint64 GetTotalSize() const
+	uint64 GetTotalSize() const
 	{
 		return TotalContentSize;
 	}
