GameFeaturesSubsystem IAD API
- Support package tag sets and more progress improvements

diff --git a/Engine/Plugins/Runtime/GameFeatures/Source/GameFeatures/Private/GameFeaturePluginStateMachine.cpp b/Engine/Plugins/Runtime/GameFeatures/Source/GameFeatures/Private/GameFeaturePluginStateMachine.cpp
--- a/Engine/Plugins/Runtime/GameFeatures/Source/GameFeatures/Private/GameFeaturePluginStateMachine.cpp
+++ b/Engine/Plugins/Runtime/GameFeatures/Source/GameFeatures/Private/GameFeaturePluginStateMachine.cpp
@@ -2650,7 +2650,7 @@ struct FGameFeaturePluginState_AssetDependencyStreaming : public FGameFeaturePlu
 	struct FIoStoreOnDemandProgress
 	{
 		FName InstallBundle;
-		float Progress = 0.0f;
+		UE::IoStore::FOnDemandInstallProgress Progress;
 	};
 
 	// Required for callback lifetime safety
@@ -2784,16 +2784,17 @@ struct FGameFeaturePluginState_AssetDependencyStreaming : public FGameFeaturePlu
 						return;
 					}
 
-					const double RequestProgress = Progress.GetRelativeProgress();
-
 					FIoStoreOnDemandProgress* MyProgress = Algo::FindBy(
 						LambdaOnDemandContext->Progress, InstallBundle, &FIoStoreOnDemandProgress::InstallBundle);
 					check(MyProgress);
-					MyProgress->Progress = RequestProgress;
+					MyProgress->Progress = Progress;
 
-					const float SumProgress = Algo::TransformAccumulate(
-						LambdaOnDemandContext->Progress, [](const FIoStoreOnDemandProgress& In) { return In.Progress; }, 0.0f);
-					const float OverallProgress = SumProgress / LambdaOnDemandContext->Progress.Num();
+					const UE::IoStore::FOnDemandInstallProgress SumProgress = Algo::TransformAccumulate(
+						LambdaOnDemandContext->Progress,
+						&FIoStoreOnDemandProgress::Progress,
+						UE::IoStore::FOnDemandInstallProgress(),
+						&UE::IoStore::FOnDemandInstallProgress::Combine);
+					const float OverallProgress = SumProgress.GetRelativeProgress();
 
 					StateProperties.OnFeatureStateProgressUpdate.ExecuteIfBound(OverallProgress);
 				},
@@ -3820,6 +3821,11 @@ const FString& UGameFeaturePluginStateMachine::GetPluginURL() const
 	return StateProperties.PluginIdentifier.GetFullPluginURL();
 }
 
+const FGameFeatureProtocolMetadata& UGameFeaturePluginStateMachine::GetProtocolMetadata() const
+{
+	return StateProperties.ProtocolMetadata;
+}
+
 const FGameFeatureProtocolOptions& UGameFeaturePluginStateMachine::GetProtocolOptions() const
 {
 	return StateProperties.ProtocolOptions;
diff --git a/Engine/Plugins/Runtime/GameFeatures/Source/GameFeatures/Private/GameFeaturePluginStateMachine.h b/Engine/Plugins/Runtime/GameFeatures/Source/GameFeatures/Private/GameFeaturePluginStateMachine.h
--- a/Engine/Plugins/Runtime/GameFeatures/Source/GameFeatures/Private/GameFeaturePluginStateMachine.h
+++ b/Engine/Plugins/Runtime/GameFeatures/Source/GameFeatures/Private/GameFeaturePluginStateMachine.h
@@ -506,6 +506,9 @@ public:
 	/** Returns the URL */
 	const FString& GetPluginURL() const;
 
+	/** Returns protocol medata */
+	const FGameFeatureProtocolMetadata& GetProtocolMetadata() const;
+
 	/** Returns any protocol options */
 	const FGameFeatureProtocolOptions& GetProtocolOptions() const;
 
diff --git a/Engine/Plugins/Runtime/GameFeatures/Source/GameFeatures/Private/GameFeaturesSubsystem.cpp b/Engine/Plugins/Runtime/GameFeatures/Source/GameFeatures/Private/GameFeaturesSubsystem.cpp
--- a/Engine/Plugins/Runtime/GameFeatures/Source/GameFeatures/Private/GameFeaturesSubsystem.cpp
+++ b/Engine/Plugins/Runtime/GameFeatures/Source/GameFeatures/Private/GameFeaturesSubsystem.cpp
@@ -1,6 +1,7 @@
 // Copyright Epic Games, Inc. All Rights Reserved.
 
 #include "GameFeaturesSubsystem.h"
+#include "Algo/Accumulate.h"
 #include "Algo/TopologicalSort.h"
 #include "AssetRegistry/AssetData.h"
 #include "AssetRegistry/IAssetRegistry.h"
@@ -2836,7 +2837,8 @@ struct FGameFeatureAssetStreamContext : public FGameFeaturePluginAssetStreamHand
 	TUniqueFunction<void(const UE::GameFeatures::FResult&)> OnComplete;
 	TUniqueFunction<void(float)> OnProgress;
 	UE::GameFeatures::FResult Result = MakeValue();
-	float Progress = 0.0f;
+	TArray<UE::IoStore::FOnDemandInstallProgress> ProgressList;
+	int32 PendingInstalls = 0;
 	bool bIsComplete = false;
 	UE::IoStore::FOnDemandCancellationToken CancellationToken;
 
@@ -2856,7 +2858,9 @@ struct FGameFeatureAssetStreamContext : public FGameFeaturePluginAssetStreamHand
 
 	virtual float GetProgress() const override
 	{
-		return Progress;
+		const UE::IoStore::FOnDemandInstallProgress SumProgress = Algo::Accumulate(
+			ProgressList, UE::IoStore::FOnDemandInstallProgress(), &UE::IoStore::FOnDemandInstallProgress::Combine);
+		return SumProgress.GetRelativeProgress();
 	}
 
 	virtual void Cancel() override
@@ -2881,11 +2885,7 @@ struct FGameFeatureAssetStreamContext : public FGameFeaturePluginAssetStreamHand
 	void SetComplete(UE::GameFeatures::FResult&& InResult)
 	{
 		Result = MoveTemp(InResult);
-		bIsComplete = true;
-		if (OnComplete)
-		{
-			OnComplete(Result);
-		}
+		SetComplete();
 	}
 
 	void SetCompleteCanceled()
@@ -2893,14 +2893,10 @@ struct FGameFeatureAssetStreamContext : public FGameFeaturePluginAssetStreamHand
 		Result = MakeError(FString::Printf(TEXT("%.*s%s"),
 			ErrorNamespace.Len(), ErrorNamespace.GetData(),
 			TEXT("Canceled")));
-		bIsComplete = true;
-		if (OnComplete)
-		{
-			OnComplete(Result);
-		}
+		SetComplete();
 	}
 
-	void Start(TArrayView<FName> AssetPackageNames)
+	void Start(TArrayView<FName> AssetPackageNames, EGameFeatureAssetStreamPriority Priority)
 	{
 		if (CancellationToken.IsCanceled())
 		{
@@ -2944,6 +2940,15 @@ struct FGameFeatureAssetStreamContext : public FGameFeaturePluginAssetStreamHand
 				return;
 			}
 
+			if (Machine->GetPluginIdentifier().GetPluginProtocol() != EGameFeaturePluginProtocol::InstallBundle)
+			{
+				UE_LOGFMT(LogGameFeatures, Error, "Unsupported protocol for GFP {MountPoint}", ("MountPoint", MountPoint));
+				UE::GameFeatures::FResult ErrorResult = MakeError(FString::Printf(TEXT("%.*s%s"),
+					ErrorNamespace.Len(), ErrorNamespace.GetData(), TEXT("Unsupported_GFP_Protocol")));
+				SetComplete(MoveTemp(ErrorResult));
+				return;
+			}
+
 			// Check installed and not moving to a state less than installed
 			if (Machine->GetCurrentState() >= EGameFeaturePluginState::Installed &&
 				Machine->GetDestination() > EGameFeaturePluginState::StatusKnown)
@@ -2957,6 +2962,107 @@ struct FGameFeatureAssetStreamContext : public FGameFeaturePluginAssetStreamHand
 		}
 
 		// Install the packages
+		UE::IoStore::FOnDemandInstallArgs InstallArgs;
+		InstallArgs.PackageIds.Reserve(AssetPackageNames.Num());
+		Algo::Transform(AssetPackageNames, InstallArgs.PackageIds, &FPackageId::FromName);
+
+		Start(MakeArrayView(&InstallArgs, 1), Priority);
+	}
+
+	void StartByTag(TArrayView<FString> Tags, EGameFeatureAssetStreamPriority Priority)
+	{
+		if (CancellationToken.IsCanceled())
+		{
+			SetCompleteCanceled();
+			return;
+		}
+
+		// This path streams matching packages from all mounted plugins so plugin state can't
+		// be verfied here. The user will just have to deal with low level errors if they occur.
+
+		UE::IoStore::FOnDemandInstallArgs InstallArgs;
+		InstallArgs.TagSets = Tags;
+
+		Start(MakeArrayView(&InstallArgs, 1), Priority);
+	}
+
+	void StartByTag(TArrayView<FString> PluginURLs, TArrayView<FString> Tags, EGameFeatureAssetStreamPriority Priority)
+	{
+		if (CancellationToken.IsCanceled())
+		{
+			SetCompleteCanceled();
+			return;
+		}
+
+		UGameFeaturesSubsystem& GFPSubSys = UGameFeaturesSubsystem::Get();
+
+		TArray<FName> ContainerMountIds;
+
+		// Verify that the we can stream from these GFPs
+		for (const FString& PluginURL : PluginURLs)
+		{
+			const UGameFeaturePluginStateMachine* Machine = GFPSubSys.FindGameFeaturePluginStateMachine(PluginURL);
+			if (!Machine)
+			{
+				FStringView PluginPath;
+				UGameFeaturesSubsystem::ParsePluginURL(PluginURL, nullptr, &PluginPath);
+
+				UE_LOGFMT(LogGameFeatures, Error, "Could not find GFP {PluginPath}", ("PluginPath", PluginPath));
+				UE::GameFeatures::FResult ErrorResult = MakeError(FString::Printf(TEXT("%.*s%s"),
+					ErrorNamespace.Len(), ErrorNamespace.GetData(), TEXT("Unknown_GFP_Mount_Point")));
+				SetComplete(MoveTemp(ErrorResult));
+				return;
+			}
+
+			FStringView PluginName = Machine->GetPluginIdentifier().GetPluginName();
+
+			if (Machine->GetPluginIdentifier().GetPluginProtocol() != EGameFeaturePluginProtocol::InstallBundle)
+			{
+				UE_LOGFMT(LogGameFeatures, Error, "Unsupported protocol for GFP {PluginName}", ("PluginName", PluginName));
+				UE::GameFeatures::FResult ErrorResult = MakeError(FString::Printf(TEXT("%.*s%s"),
+					ErrorNamespace.Len(), ErrorNamespace.GetData(), TEXT("Unsupported_GFP_Protocol")));
+				SetComplete(MoveTemp(ErrorResult));
+				return;
+			}
+
+			// Check installed and not moving to a state less than installed
+			if (Machine->GetCurrentState() >= EGameFeaturePluginState::Installed &&
+				Machine->GetDestination() > EGameFeaturePluginState::StatusKnown)
+			{
+				UE_LOGFMT(LogGameFeatures, Error, "Invalid GFP state for {PluginName}", ("PluginName", PluginName));
+				UE::GameFeatures::FResult ErrorResult = MakeError(FString::Printf(TEXT("%.*s%s"),
+					ErrorNamespace.Len(), ErrorNamespace.GetData(), TEXT("GFP_Not_Ready")));
+				SetComplete(MoveTemp(ErrorResult));
+				return;
+			}
+
+			// Get mount ids of containers to install assets from
+			const FInstallBundlePluginProtocolMetaData& Metadata = Machine->GetProtocolMetadata().GetSubtype<FInstallBundlePluginProtocolMetaData>();
+			for (const FName InstallBundle : Metadata.InstallBundlesWithAssetDependencies)
+			{
+				ContainerMountIds.AddUnique(InstallBundle);
+			}
+		}
+
+		if (ContainerMountIds.IsEmpty())
+		{
+			SetComplete(); // Success?! There's nothing to install
+		}
+
+		TArray<UE::IoStore::FOnDemandInstallArgs> InstallArgsList;
+		InstallArgsList.Reserve(ContainerMountIds.Num());
+		for (const FName MountId : ContainerMountIds)
+		{
+			UE::IoStore::FOnDemandInstallArgs& InstallArgs = InstallArgsList.Emplace_GetRef();
+			InstallArgs.MountId = MountId.ToString();
+			InstallArgs.TagSets = Tags;
+		}
+
+		Start(InstallArgsList, Priority);
+	}
+
+	void Start(TArrayView<UE::IoStore::FOnDemandInstallArgs> InstallArgList, EGameFeatureAssetStreamPriority Priority)
+	{
 		UE::IoStore::FIoStoreOnDemandModule* IoStoreOnDemandModule = FModuleManager::Get().GetModulePtr<UE::IoStore::FIoStoreOnDemandModule>(TEXT("IoStoreOnDemand"));
 		if (!IoStoreOnDemandModule)
 		{
@@ -2974,65 +3080,114 @@ struct FGameFeatureAssetStreamContext : public FGameFeaturePluginAssetStreamHand
 			return;
 		}
 
-		UE::IoStore::FOnDemandInstallArgs InstallArgs;
-		InstallArgs.Options |= UE::IoStore::EOnDemandInstallOptions::CallbackOnGameThread;
-		InstallArgs.ContentHandle = AssetRef.ContentHandle;
+		check(InstallArgList.Num() > 0);
 
-		InstallArgs.PackageIds.Reserve(AssetPackageNames.Num());
-		Algo::Transform(AssetPackageNames, InstallArgs.PackageIds, &FPackageId::FromName);
+		PendingInstalls = InstallArgList.Num();
+		ProgressList.Empty(PendingInstalls);
 
-		IoStoreOnDemandModule->Install(MoveTemp(InstallArgs),
-			[Context = SharedThis(this)](const UE::IoStore::FOnDemandInstallResult& OnDemandInstallResult)
-			{
-				Context->OnInstallComplete(OnDemandInstallResult);
-			},
-			[Context = SharedThis(this)](const UE::IoStore::FOnDemandInstallProgress& InProgress)
-			{
-				Context->OnInstallProgress(InProgress);
-			},
-			&CancellationToken);
+		for (UE::IoStore::FOnDemandInstallArgs& InstallArgs : InstallArgList)
+		{
+			const int32 RequestIndex = ProgressList.Emplace();
+
+			InstallArgs.ContentHandle = AssetRef.ContentHandle;
+			InstallArgs.Options |= UE::IoStore::EOnDemandInstallOptions::CallbackOnGameThread;
+
+			// TODO: convert Priority to IOStore priority here
+
+			IoStoreOnDemandModule->Install(MoveTemp(InstallArgs),
+				[Context = SharedThis(this), RequestIndex](const UE::IoStore::FOnDemandInstallResult& OnDemandInstallResult)
+				{
+					Context->OnInstallComplete(OnDemandInstallResult, RequestIndex);
+				},
+				[Context = SharedThis(this), RequestIndex](const UE::IoStore::FOnDemandInstallProgress& InProgress)
+				{
+					Context->OnInstallProgress(InProgress, RequestIndex);
+				},
+				&CancellationToken);
+		}
 	}
 
-	void OnInstallComplete(const UE::IoStore::FOnDemandInstallResult& OnDemandInstallResult)
+	void OnInstallComplete(const UE::IoStore::FOnDemandInstallResult& OnDemandInstallResult, const int32 RequestIndex)
 	{
-		if (!OnDemandInstallResult.Status.IsOk())
+		if (!OnDemandInstallResult.Status.IsOk() && !Result.HasError())
 		{
 			const FString StatusString = OnDemandInstallResult.Status.ToString();
 			FString ErrorCode = GetIoErrorText(OnDemandInstallResult.Status.GetErrorCode());
 			ErrorCode.ReplaceCharInline(TEXT(' '), TEXT('_'), ESearchCase::CaseSensitive);
 
-			UE::GameFeatures::FResult ErrorResult(MakeError(
+			Result = UE::GameFeatures::FResult(MakeError(
 				FString::Printf(TEXT("%.*s%s.%s"),
 					ErrorNamespace.Len(), ErrorNamespace.GetData(), TEXT("IoStoreOnDemand.OnComplete"), *ErrorCode)),
 				FText::AsCultureInvariant(StatusString));
-			SetComplete(MoveTemp(ErrorResult));
-			return;
+
+			// Something went wrong, kill any other requests
+			Cancel();
 		}
 
-		SetComplete();
+		--PendingInstalls;
+		if (PendingInstalls == 0)
+		{
+			SetComplete();
+		}
 	}
 
-	void OnInstallProgress(const UE::IoStore::FOnDemandInstallProgress& InProgress)
+	void OnInstallProgress(const UE::IoStore::FOnDemandInstallProgress& InProgress, const int32 RequestIndex)
 	{
-		Progress = InProgress.GetRelativeProgress();
+		ProgressList[RequestIndex] = InProgress;
 		if (OnProgress)
 		{
-			OnProgress(Progress);
+			OnProgress(GetProgress());
 		}
 	}
 };
 
-TSharedRef<FGameFeaturePluginAssetStreamHandle> UGameFeaturesSubsystem::DownloadStreamedAssets(FGameFeaturePluginStreamedAssetsRef AssetRef, TArrayView<FName> AssetPackageNames,
-	TUniqueFunction<void(const UE::GameFeatures::FResult&)> OnComplete /*= nullptr*/, TUniqueFunction<void(float)> OnProgress /*= nullptr*/)
+TSharedRef<FGameFeaturePluginAssetStreamHandle> UGameFeaturesSubsystem::DownloadStreamedAssets(
+	FGameFeaturePluginStreamedAssetsRef AssetRef,
+	TArrayView<FName> AssetPackageNames,
+	EGameFeatureAssetStreamPriority Priority /*= EGameFeatureAssetStreamPriority::Normal*/,
+	TUniqueFunction<void(const UE::GameFeatures::FResult&)> OnComplete /*= nullptr*/, 
+	TUniqueFunction<void(float)> OnProgress /*= nullptr*/)
 {
 	TSharedRef<FGameFeatureAssetStreamContext> Context = MakeShared<FGameFeatureAssetStreamContext>(MoveTemp(AssetRef));
 	Context->OnComplete = MoveTemp(OnComplete);
 	Context->OnProgress = MoveTemp(OnProgress);
-	Context->Start(AssetPackageNames);
+	Context->Start(AssetPackageNames, Priority);
 
 	return Context;
 }
 
+TSharedRef<FGameFeaturePluginAssetStreamHandle> UGameFeaturesSubsystem::DownloadStreamedAssetsByTag(
+	FGameFeaturePluginStreamedAssetsRef AssetRef,
+	TArrayView<FString> Tags,
+	EGameFeatureAssetStreamPriority Priority /*= EGameFeatureAssetStreamPriority::Normal*/,
+	TUniqueFunction<void(const UE::GameFeatures::FResult&)> OnComplete /*= nullptr*/,
+	TUniqueFunction<void(float)> OnProgress /*= nullptr*/)
+{
+	// TODO: not sure if the overload makes a lot of sense. What's the use case?
+
+	TSharedRef<FGameFeatureAssetStreamContext> Context = MakeShared<FGameFeatureAssetStreamContext>(MoveTemp(AssetRef));
+	Context->OnComplete = MoveTemp(OnComplete);
+	Context->OnProgress = MoveTemp(OnProgress);
+	Context->StartByTag(Tags, Priority);
+
+	return Context;
+}
+
+TSharedRef<FGameFeaturePluginAssetStreamHandle> DownloadStreamedAssetsByTag(
+	FGameFeaturePluginStreamedAssetsRef AssetRef,
+	TArrayView<FString> PluginURLs,
+	TArrayView<FString> Tags,
+	EGameFeatureAssetStreamPriority Priority /*= EGameFeatureAssetStreamPriority::Normal*/,
+	TUniqueFunction<void(const UE::GameFeatures::FResult&)> OnComplete /*= nullptr*/,
+	TUniqueFunction<void(float)> OnProgress /*= nullptr*/)
+{
+	TSharedRef<FGameFeatureAssetStreamContext> Context = MakeShared<FGameFeatureAssetStreamContext>(MoveTemp(AssetRef));
+	Context->OnComplete = MoveTemp(OnComplete);
+	Context->OnProgress = MoveTemp(OnProgress);
+	Context->StartByTag(PluginURLs, Tags, Priority);
+	return Context;
+}
+
 UGameFeaturePluginStateMachine* UGameFeaturesSubsystem::FindGameFeaturePluginStateMachine(const FString& PluginURL) const
 {
 	FGameFeaturePluginIdentifier FindPluginIdentifier(PluginURL);
diff --git a/Engine/Plugins/Runtime/GameFeatures/Source/GameFeatures/Public/GameFeaturesSubsystem.h b/Engine/Plugins/Runtime/GameFeatures/Source/GameFeatures/Public/GameFeaturesSubsystem.h
--- a/Engine/Plugins/Runtime/GameFeatures/Source/GameFeatures/Public/GameFeaturesSubsystem.h
+++ b/Engine/Plugins/Runtime/GameFeatures/Source/GameFeatures/Public/GameFeaturesSubsystem.h
@@ -297,6 +297,15 @@ struct FGameFeaturePluginAssetStreamHandle : public FGameFeaturePluginAsyncHandl
 	virtual FGameFeaturePluginStreamedAssetsRef GetAssetsRef() const = 0;
 };
 
+enum class EGameFeatureAssetStreamPriority : int8
+{
+	Lowest = -128,
+	BelowNormal = -1,
+	Normal = 0,
+	AboveNormal = 1,
+	Highest = 127
+};
+
 /** Struct used to transform a GameFeaturePlugin URL into something that can uniquely identify the GameFeaturePlugin
     without including any transient data being passed in through the URL */
 USTRUCT()
@@ -683,11 +692,29 @@ public:
 	friend struct FGameFeaturePluginPredownloadContext;
 
 	/**
+	 * ***************EXPERIMENTAL***************** 
 	 * Download streamed assets from installed GFPs. The assets cannot be loaded unless the plugin is also mounted.
 	 * Assets are gauranteed to be available as long as AssetRef is valid.
 	 */
-	TSharedRef<FGameFeaturePluginAssetStreamHandle> DownloadStreamedAssets(FGameFeaturePluginStreamedAssetsRef AssetRef, TArrayView<FName> AssetPackageNames,
-		TUniqueFunction<void(const UE::GameFeatures::FResult&)> OnComplete = nullptr, TUniqueFunction<void(float)> OnProgress = nullptr);
+	TSharedRef<FGameFeaturePluginAssetStreamHandle> DownloadStreamedAssets(
+		FGameFeaturePluginStreamedAssetsRef AssetRef, 
+		TArrayView<FName> AssetPackageNames,
+		EGameFeatureAssetStreamPriority Priority = EGameFeatureAssetStreamPriority::Normal,
+		TUniqueFunction<void(const UE::GameFeatures::FResult&)> OnComplete = nullptr, 
+		TUniqueFunction<void(float)> OnProgress = nullptr);
+	TSharedRef<FGameFeaturePluginAssetStreamHandle> DownloadStreamedAssetsByTag(
+		FGameFeaturePluginStreamedAssetsRef AssetRef,
+		TArrayView<FString> Tags,
+		EGameFeatureAssetStreamPriority Priority = EGameFeatureAssetStreamPriority::Normal,
+		TUniqueFunction<void(const UE::GameFeatures::FResult&)> OnComplete = nullptr,
+		TUniqueFunction<void(float)> OnProgress = nullptr);
+	TSharedRef<FGameFeaturePluginAssetStreamHandle> DownloadStreamedAssetsByTag(
+		FGameFeaturePluginStreamedAssetsRef AssetRef,
+		TArrayView<FString> PluginURLs,
+		TArrayView<FString> Tags,
+		EGameFeatureAssetStreamPriority Priority = EGameFeatureAssetStreamPriority::Normal,
+		TUniqueFunction<void(const UE::GameFeatures::FResult&)> OnComplete = nullptr,
+		TUniqueFunction<void(float)> OnProgress = nullptr);
 	friend struct FGameFeatureAssetStreamContext;
 
 	/** Determine the initial feature state for a built-in plugin */
diff --git a/Engine/Source/Runtime/Experimental/IoStore/OnDemand/Public/IO/IoStoreOnDemand.h b/Engine/Source/Runtime/Experimental/IoStore/OnDemand/Public/IO/IoStoreOnDemand.h
--- a/Engine/Source/Runtime/Experimental/IoStore/OnDemand/Public/IO/IoStoreOnDemand.h
+++ b/Engine/Source/Runtime/Experimental/IoStore/OnDemand/Public/IO/IoStoreOnDemand.h
@@ -353,6 +353,13 @@ struct FOnDemandInstallProgress
 	/** The size currently installed/downloaded (<= TotalInstallSize). */
 	uint64 CurrentInstallSize = 0;
 
+	FOnDemandInstallProgress& Combine(const FOnDemandInstallProgress& Other)
+	{
+		TotalContentSize += Other.TotalContentSize;
+		TotalInstallSize += Other.TotalInstallSize;
+		CurrentInstallSize += Other.CurrentInstallSize;
+		return *this;
+	}
 
 	uint64 GetTotalDownloadSize() const
 	{
