Avoid checking the timestamp to see if we need to invalidate the cache, as this can be expensive to call
Move the lock around the cache to a read lock and only grab a write lock when adding to the cache. This should allow during initial load when we go wide to *block* less on mutexes as well as later when we try to go more async to not cause issues trying to read from the cache when no one is writing to it

diff --git a/Engine/Plugins/Runtime/GameFeatures/Source/GameFeatures/Private/GameFeaturePluginStateMachine.cpp b/Engine/Plugins/Runtime/GameFeatures/Source/GameFeatures/Private/GameFeaturePluginStateMachine.cpp
--- a/Engine/Plugins/Runtime/GameFeatures/Source/GameFeatures/Private/GameFeaturePluginStateMachine.cpp
+++ b/Engine/Plugins/Runtime/GameFeatures/Source/GameFeatures/Private/GameFeaturePluginStateMachine.cpp
@@ -1946,8 +1946,6 @@ struct FGameFeaturePluginState_Unmounting : public FGameFeaturePluginState
 
 		if (StateProperties.bAddedPluginToManager)
 		{
-			UGameFeaturesSubsystem::Get().PruneCachedGameFeaturePluginDetails(StateProperties.PluginIdentifier.GetFullPluginURL(), StateProperties.PluginInstalledFilename);
-
 			verify(IPluginManager::Get().RemoveFromPluginsList(StateProperties.PluginInstalledFilename));
 			StateProperties.bAddedPluginToManager = false;
 		}
diff --git a/Engine/Plugins/Runtime/GameFeatures/Source/GameFeatures/Private/GameFeaturesSubsystem.cpp b/Engine/Plugins/Runtime/GameFeatures/Source/GameFeatures/Private/GameFeaturesSubsystem.cpp
--- a/Engine/Plugins/Runtime/GameFeatures/Source/GameFeatures/Private/GameFeaturesSubsystem.cpp
+++ b/Engine/Plugins/Runtime/GameFeatures/Source/GameFeatures/Private/GameFeaturesSubsystem.cpp
@@ -461,6 +461,7 @@ void UGameFeaturesSubsystem::Initialize(FSubsystemCollectionBase& Collection)
 		ECVF_Cheat);
 
 	GetExplanationForUnavailablePackageDelegateHandle = IPluginManager::Get().GetExplanationForUnavailablePackageWithPluginInfoDelegate().AddUObject(this, &UGameFeaturesSubsystem::GetExplanationForUnavailablePackage);
+	OnPluginUnmountedDelegateHandle = IPluginManager::Get().OnPluginUnmounted().AddWeakLambda(this, [this] (const IPlugin& Plugin) { PruneCachedGameFeaturePluginDetails(/* doesnt use URL */ FString(""), Plugin.GetDescriptorFileName()); });
 }
 
 void UGameFeaturesSubsystem::Deinitialize()
@@ -468,6 +469,7 @@ void UGameFeaturesSubsystem::Deinitialize()
 	UE_LOG(LogGameFeatures, Log, TEXT("Shutting down game features subsystem"));
 	IPluginManager::Get().GetExplanationForUnavailablePackageWithPluginInfoDelegate().Remove(GetExplanationForUnavailablePackageDelegateHandle);
 	GetExplanationForUnavailablePackageDelegateHandle.Reset();
+	OnPluginUnmountedDelegateHandle.Reset();
 
 	if ((GameSpecificPolicies != nullptr) && bInitializedPolicyManager)
 	{
@@ -2345,26 +2347,14 @@ bool UGameFeaturesSubsystem::GetGameFeaturePluginDetailsInternal(const FString&
 	// GFPs are implemented with a plugin so FPluginReferenceDescriptor doesn't know anything about them.
 	// Need a better way of storing GFP specific plugin data...
 
-	FDateTime FileTimeStamp;
 	if (UE::GameFeatures::GCachePluginDetails)
 	{
-		FScopeLock Lock(&CachedGameFeaturePluginDetailsLock);
+		FReadScopeLock ReadLock(CachedGameFeaturePluginDetailsLock);
 
-		// Note: On file systems that don't support timestamps (current time is returned instead), 
-		// the pak file layer will end up caching the mount time, so this stamp will still be valid as long as the uplugin
-		// is in a pak and the pak is mounted.
-		FileTimeStamp = IFileManager::Get().GetTimeStamp(*PluginDescriptorFilename);
 		if (FCachedGameFeaturePluginDetails* ExistingDetails = CachedPluginDetailsByFilename.Find(PluginDescriptorFilename))
 		{
-			if (ExistingDetails->TimeStamp == FileTimeStamp)
-			{
-				OutPluginDetails = ExistingDetails->Details;
-				return true;
-			}
-			else
-			{
-				CachedPluginDetailsByFilename.Remove(PluginDescriptorFilename);
-			}
+			OutPluginDetails = ExistingDetails->Details;
+			return true;
 		}
 	}
 
@@ -2471,8 +2461,8 @@ bool UGameFeaturesSubsystem::GetGameFeaturePluginDetailsInternal(const FString&
 
 	if (UE::GameFeatures::GCachePluginDetails)
 	{
-		FScopeLock Lock(&CachedGameFeaturePluginDetailsLock);
-		CachedPluginDetailsByFilename.Add(PluginDescriptorFilename, FCachedGameFeaturePluginDetails(OutPluginDetails, FileTimeStamp));
+		FWriteScopeLock	WriteLock(CachedGameFeaturePluginDetailsLock);
+		CachedPluginDetailsByFilename.Add(PluginDescriptorFilename, FCachedGameFeaturePluginDetails(OutPluginDetails));
 	}
 
 	return true;
@@ -2480,6 +2470,7 @@ bool UGameFeaturesSubsystem::GetGameFeaturePluginDetailsInternal(const FString&
 
 void UGameFeaturesSubsystem::PruneCachedGameFeaturePluginDetails(const FString& PluginURL, const FString& PluginDescriptorFilename) const
 {
+	FWriteScopeLock WriteLock(CachedGameFeaturePluginDetailsLock);
 	CachedPluginDetailsByFilename.Remove(PluginDescriptorFilename);
 }
 
diff --git a/Engine/Plugins/Runtime/GameFeatures/Source/GameFeatures/Public/GameFeaturesSubsystem.h b/Engine/Plugins/Runtime/GameFeatures/Source/GameFeatures/Public/GameFeaturesSubsystem.h
--- a/Engine/Plugins/Runtime/GameFeatures/Source/GameFeatures/Public/GameFeaturesSubsystem.h
+++ b/Engine/Plugins/Runtime/GameFeatures/Source/GameFeatures/Public/GameFeaturesSubsystem.h
@@ -886,12 +886,11 @@ private:
 	struct FCachedGameFeaturePluginDetails
 	{
 		FGameFeaturePluginDetails Details;
-		FDateTime TimeStamp;
 		FCachedGameFeaturePluginDetails() = default;
-		FCachedGameFeaturePluginDetails(const FGameFeaturePluginDetails& InDetails, const FDateTime& InTimeStamp) : Details(InDetails), TimeStamp(InTimeStamp) {}
+		FCachedGameFeaturePluginDetails(const FGameFeaturePluginDetails& InDetails) : Details(InDetails) {}
 	};
 	mutable TMap<FString, FCachedGameFeaturePluginDetails> CachedPluginDetailsByFilename;
-	mutable FCriticalSection CachedGameFeaturePluginDetailsLock;
+	mutable FRWLock CachedGameFeaturePluginDetailsLock;
 
 	UPROPERTY()
 	TArray<TObjectPtr<UObject>> Observers;
@@ -904,6 +903,7 @@ private:
 	TMap<FString, FString> UnmountedPluginNameToExplanation;
 #endif
 	FDelegateHandle GetExplanationForUnavailablePackageDelegateHandle;
+	FDelegateHandle OnPluginUnmountedDelegateHandle;
 
 	bool bInitializedPolicyManager = false;
 };
