Add GameFeaturePlugin.ForceSyncRegisterStartupPlugins (defaults to true) cvar to make the game feature state machine synchronously register plugins during engine startup even if async loading is enabled.
This is required for engine startup tasks like restoring open editor tabs to work properly, as they load data that might require plugin registration.

diff --git a/Engine/Plugins/Runtime/GameFeatures/Source/GameFeatures/Private/GameFeaturePluginStateMachine.cpp b/Engine/Plugins/Runtime/GameFeatures/Source/GameFeatures/Private/GameFeaturePluginStateMachine.cpp
--- a/Engine/Plugins/Runtime/GameFeatures/Source/GameFeatures/Private/GameFeaturePluginStateMachine.cpp
+++ b/Engine/Plugins/Runtime/GameFeatures/Source/GameFeatures/Private/GameFeaturePluginStateMachine.cpp
@@ -78,16 +78,20 @@ namespace UE::GameFeatures
 
 	static TAutoConsoleVariable<bool> CVarAsyncLoad(TEXT("GameFeaturePlugin.AsyncLoad"),
 		true,
-		TEXT("Enable to use aysnc loading"));
+		TEXT("Enable to use async loading as well async downloading and registering"));
 
 	static TAutoConsoleVariable<bool> CVarForceAsyncLoad(TEXT("GameFeaturePlugin.ForceAsyncLoad"),
 		false,
-		TEXT("Enable to force use of aysnc loading even if normally not allowed"));
+		TEXT("Enable to force use of async loading even if normally not allowed"));
 
 	static TAutoConsoleVariable<bool> CVarAllowForceMonolithicShaderLibrary(TEXT("GameFeaturePlugin.AllowForceMonolithicShaderLibrary"),
 		true,
 		TEXT("Enable to force only searching for monolithic shader libs when possible"));
 
+	static TAutoConsoleVariable<bool> CVarForceSyncRegisterStartupPlugins(TEXT("GameFeaturePlugin.ForceSyncRegisterStartupPlugins"),
+		true,
+		TEXT("If true, all plugins loaded during startup will be synchronously registered to ensure things are initialized in time, this only applies if AsyncLoad is enabled"));
+
 	static TAutoConsoleVariable<bool> CVarForceSyncLoadShaderLibrary(TEXT("GameFeaturePlugin.ForceSyncLoadShaderLibrary"),
 		true,
 		TEXT("Enable to force shaderlibs to be opened on the game thread"));
@@ -2168,6 +2172,19 @@ struct FGameFeaturePluginState_Mounting : public FGameFeaturePluginState
 		}
 	}
 
+	virtual bool UseAsyncLoading() const override
+	{
+		if (UE::GameFeatures::CVarForceSyncRegisterStartupPlugins.GetValueOnGameThread())
+		{
+			if (UGameFeaturesSubsystem::Get().GetPolicy().IsLoadingStartupPlugins())
+			{
+				return false;
+			}
+		}
+
+		return FGameFeaturePluginState::UseAsyncLoading();
+	}
+
 	virtual void BeginState() override
 	{
 		NumObservedPostMountPausers = 0;
@@ -2970,6 +2987,18 @@ struct FGameFeaturePluginState_Registering : public FGameFeaturePluginState
 		}
 	}
 
+	virtual bool UseAsyncLoading() const override
+	{
+		if (UE::GameFeatures::CVarForceSyncRegisterStartupPlugins.GetValueOnGameThread())
+		{
+			if (UGameFeaturesSubsystem::Get().GetPolicy().IsLoadingStartupPlugins())
+			{
+				return false;
+			}
+		}
+		return FGameFeaturePluginState::UseAsyncLoading();
+	}
+
 	virtual void BeginState() override
 	{
 		TRACE_CPUPROFILER_EVENT_SCOPE(GFP_Registering_Begin);
diff --git a/Engine/Plugins/Runtime/GameFeatures/Source/GameFeatures/Private/GameFeaturePluginStateMachine.h b/Engine/Plugins/Runtime/GameFeatures/Source/GameFeatures/Private/GameFeaturePluginStateMachine.h
--- a/Engine/Plugins/Runtime/GameFeatures/Source/GameFeatures/Private/GameFeaturePluginStateMachine.h
+++ b/Engine/Plugins/Runtime/GameFeatures/Source/GameFeatures/Private/GameFeaturePluginStateMachine.h
@@ -418,7 +418,7 @@ protected:
 	bool AllowIniLoading() const;
 
 	bool AllowAsyncLoading() const;
-	bool UseAsyncLoading() const;
+	virtual bool UseAsyncLoading() const;
 
 private:
 	void CleanupDeferredUpdateCallbacks() const;
diff --git a/Engine/Plugins/Runtime/GameFeatures/Source/GameFeatures/Private/GameFeaturesProjectPolicies.cpp b/Engine/Plugins/Runtime/GameFeatures/Source/GameFeatures/Private/GameFeaturesProjectPolicies.cpp
--- a/Engine/Plugins/Runtime/GameFeatures/Source/GameFeatures/Private/GameFeaturesProjectPolicies.cpp
+++ b/Engine/Plugins/Runtime/GameFeatures/Source/GameFeatures/Private/GameFeaturesProjectPolicies.cpp
@@ -4,6 +4,7 @@
 #include "GameFeaturesSubsystemSettings.h"
 #include "Interfaces/IPluginManager.h"
 #include "Misc/Paths.h"
+#include "Engine/Engine.h"
 
 #include UE_INLINE_GENERATED_CPP_BY_NAME(GameFeaturesProjectPolicies)
 
@@ -50,6 +51,23 @@ const TArray<FName> UDefaultGameFeaturesProjectPolicies::GetPreloadBundleStateFo
 	return FeatureBundles;
 }
 
+bool UGameFeaturesProjectPolicies::IsLoadingStartupPlugins() const
+{
+	if (GIsRunning && GFrameCounter > 2)
+	{
+		// Initial loading can take 2 frames
+		return false;
+	}
+
+	if (IsRunningCommandlet() && GEngine && GEngine->IsInitialized())
+	{
+		// Commandlets may not tick, so done after initialization
+		return false;
+	}
+
+	return true;
+}
+
 bool UGameFeaturesProjectPolicies::GetGameFeaturePluginURL(const TSharedRef<IPlugin>& Plugin, FString& OutPluginURL) const
 {
 	// It could still be a GFP, but state machine may not have been created for it yet
diff --git a/Engine/Plugins/Runtime/GameFeatures/Source/GameFeatures/Public/GameFeaturesProjectPolicies.h b/Engine/Plugins/Runtime/GameFeatures/Source/GameFeatures/Public/GameFeaturesProjectPolicies.h
--- a/Engine/Plugins/Runtime/GameFeatures/Source/GameFeatures/Public/GameFeaturesProjectPolicies.h
+++ b/Engine/Plugins/Runtime/GameFeatures/Source/GameFeatures/Public/GameFeaturesProjectPolicies.h
@@ -48,6 +48,10 @@ public:
 	// Actions can use this to decide what to load at runtime
 	virtual void GetGameFeatureLoadingMode(bool& bLoadClientData, bool& bLoadServerData) const { bLoadClientData = true; bLoadServerData = true; }
 
+	// Called to determine if we are still during engine startup, which can modify loading behavior.
+	// This defaults to true for the first few frames of a normal game or editor, but can be overridden.
+	virtual bool IsLoadingStartupPlugins() const;
+
 	// Called to determine the plugin URL for a given known Plugin. Can be used if the policy wants to deliver non file based URLs.
 	virtual bool GetGameFeaturePluginURL(const TSharedRef<IPlugin>& Plugin, FString& OutPluginURL) const;
 
