Fix an issue where loading into a game where many gamefeatureplugins fails can significantly slow the boot time.
Check if a dependency has already failed to register or waiting for dependencies so that we don't attempt to process it twice.

diff --git a/Engine/Plugins/Runtime/GameFeatures/Source/GameFeatures/Private/GameFeaturePluginStateMachine.cpp b/Engine/Plugins/Runtime/GameFeatures/Source/GameFeatures/Private/GameFeaturePluginStateMachine.cpp
--- a/Engine/Plugins/Runtime/GameFeatures/Source/GameFeatures/Private/GameFeaturePluginStateMachine.cpp
+++ b/Engine/Plugins/Runtime/GameFeatures/Source/GameFeatures/Private/GameFeaturePluginStateMachine.cpp
@@ -654,10 +654,9 @@ struct FTransitionDependenciesGameFeaturePluginState : public FGameFeaturePlugin
 		checkf(!StateProperties.PluginInstalledFilename.IsEmpty(), TEXT("PluginInstalledFilename must be set by the loading dependencies phase. PluginURL: %s"), *StateProperties.PluginIdentifier.GetFullPluginURL());
 		checkf(FPaths::GetExtension(StateProperties.PluginInstalledFilename) == TEXT("uplugin"), TEXT("PluginInstalledFilename must have a uplugin extension. PluginURL: %s"), *StateProperties.PluginIdentifier.GetFullPluginURL());
 
+		UGameFeaturesSubsystem& GameFeaturesSubsystem = UGameFeaturesSubsystem::Get();
 		if (!bRequestedDependencies)
 		{
-			UGameFeaturesSubsystem& GameFeaturesSubsystem = UGameFeaturesSubsystem::Get();
-
 			TArray<UGameFeaturePluginStateMachine*> Dependencies;
 			if (!TransitionPolicy::GetPluginDependencyStateMachines(StateProperties, Dependencies))
 			{
@@ -671,7 +670,6 @@ struct FTransitionDependenciesGameFeaturePluginState : public FGameFeaturePlugin
 			UE_CLOG(Dependencies.Num() > 0, LogGameFeatures, Verbose, TEXT("Found %i dependencies for %s"), Dependencies.Num(), *StateProperties.PluginName);
 
 			const bool bAllowAsyncLoading = AllowAsyncLoading();
-
 			RemainingDependencies.Reserve(Dependencies.Num());
 			for (UGameFeaturePluginStateMachine* Dependency : Dependencies)
 			{
@@ -679,8 +677,24 @@ struct FTransitionDependenciesGameFeaturePluginState : public FGameFeaturePlugin
 					TEXT("FGameFeaturePluginState::AllowAsyncLoading is false for %s but true for dependency being waited on %s"), 
 					*StateProperties.PluginName, *Dependency->GetPluginURL());
 
-				RemainingDependencies.Emplace(Dependency, MakeValue());
-				TransitionDependency(Dependency);
+				const EGameFeaturePluginState DependencyState = Dependency->GetCurrentState();
+				const bool bIsInErrorState = DependencyState == EGameFeaturePluginState::ErrorRegistering || DependencyState == EGameFeaturePluginState::ErrorWaitingForDependencies;
+				if (TransitionPolicy::ShouldWaitForDependencies() && bIsInErrorState && Dependency->IsErrorStateUnrecoverable())
+				{
+					UE_LOG(LogGameFeatures, Error, TEXT("Dependency %s failed and is in %s"), *Dependency->GetPluginURL(), *UE::GameFeatures::ToString(DependencyState));
+					StateStatus.SetTransitionError(TransitionPolicy::GetErrorState(), GetErrorResult(TEXT("Failed_Dependency_Transition")));
+					if (UGameFeaturePluginStateMachine* CurrentMachine = GameFeaturesSubsystem.FindGameFeaturePluginStateMachine(StateProperties.PluginIdentifier))
+					{
+						UE_LOG(LogGameFeatures, Error, TEXT("Setting %s to be in unrecoverable error as one of their dependencies is in unrecoverable error"), *StateProperties.PluginIdentifier.GetFullPluginURL());
+						CurrentMachine->SetUnrecoverableError();
+					}
+					return;
+				}
+				else
+				{
+					RemainingDependencies.Emplace(Dependency, MakeValue());
+					TransitionDependency(Dependency);
+				}
 			}
 		}
 
@@ -698,6 +712,14 @@ struct FTransitionDependenciesGameFeaturePluginState : public FGameFeaturePlugin
 			{
 				UE_LOG(LogGameFeatures, Error, TEXT("Dependency %s failed to transition with error %s"), *RemainingDependency->GetPluginURL(), *Pair.Value.GetError());
 				StateStatus.SetTransitionError(TransitionPolicy::GetErrorState(), GetErrorResult(TEXT("Failed_Dependency_Transition")));
+				if (RemainingDependency->IsErrorStateUnrecoverable())
+				{
+					if (UGameFeaturePluginStateMachine* CurrentMachine = GameFeaturesSubsystem.FindGameFeaturePluginStateMachine(StateProperties.PluginIdentifier))
+					{
+						UE_LOG(LogGameFeatures, Error, TEXT("Setting %s to be in unrecoverable error as one of their dependencies is in unrecoverable error"), *StateProperties.PluginIdentifier.GetFullPluginURL());
+						CurrentMachine->SetUnrecoverableError();
+					}
+				}
 				return;
 			}
 		}
@@ -3139,6 +3161,11 @@ struct FGameFeaturePluginState_Registering : public FGameFeaturePluginState
 
 			// The gamefeaturedata does not exist. The pak file may not be openable or this is a builtin plugin where the pak file does not exist.
 			StateStatus.SetTransitionError(EGameFeaturePluginState::ErrorRegistering, GetErrorResult(TEXT("Plugin_Missing_GameFeatureData")));
+			if (UGameFeaturePluginStateMachine* CurrentMachine = UGameFeaturesSubsystem::Get().FindGameFeaturePluginStateMachine(StateProperties.PluginIdentifier))
+			{
+				UE_LOG(LogGameFeatures, Error, TEXT("Setting %s to be in unrecoverable error as GameFeatureData is missing"), *StateProperties.PluginIdentifier.GetFullPluginURL());
+				CurrentMachine->SetUnrecoverableError();
+			}
 		}
 	}
 };
@@ -3961,6 +3988,16 @@ const FGameFeaturePluginStateMachineProperties& UGameFeaturePluginStateMachine::
 	return StateProperties;
 }
 
+bool UGameFeaturePluginStateMachine::IsErrorStateUnrecoverable() const
+{
+	return bIsInUnrecoverableError;
+}
+
+void UGameFeaturePluginStateMachine::SetUnrecoverableError()
+{
+	bIsInUnrecoverableError = true;
+}
+
 bool UGameFeaturePluginStateMachine::IsValidTransitionState(EGameFeaturePluginState InState) const
 {
 	check(InState != EGameFeaturePluginState::MAX);
diff --git a/Engine/Plugins/Runtime/GameFeatures/Source/GameFeatures/Private/GameFeaturePluginStateMachine.h b/Engine/Plugins/Runtime/GameFeatures/Source/GameFeatures/Private/GameFeaturePluginStateMachine.h
--- a/Engine/Plugins/Runtime/GameFeatures/Source/GameFeatures/Private/GameFeaturePluginStateMachine.h
+++ b/Engine/Plugins/Runtime/GameFeatures/Source/GameFeatures/Private/GameFeaturePluginStateMachine.h
@@ -557,6 +557,12 @@ public:
 
 	const FGameFeaturePluginStateMachineProperties& GetProperties() const;
 
+	/** Returns true if the state machine is in an unrecoverable error state. */
+	bool IsErrorStateUnrecoverable() const;
+
+	/** Sets that the state machine is in an unrecoverable error. Used when the GFD is missing as that can't be recovered from. */
+	void SetUnrecoverableError();
+
 private:
 	/** Returns true if the specified state is not a transition state */
 	bool IsValidTransitionState(EGameFeaturePluginState InState) const;
@@ -585,4 +591,7 @@ private:
 
 	/** True when we are currently executing UpdateStateMachine, to avoid reentry */
 	bool bInUpdateStateMachine;
+
+	/** True when the state machine can not transition out of the current error state or if requested to would likely fail */
+	bool bIsInUnrecoverableError = false;
 };
diff --git a/Engine/Plugins/Runtime/GameFeatures/Source/GameFeatures/Private/GameFeaturesSubsystem.cpp b/Engine/Plugins/Runtime/GameFeatures/Source/GameFeatures/Private/GameFeaturesSubsystem.cpp
--- a/Engine/Plugins/Runtime/GameFeatures/Source/GameFeatures/Private/GameFeaturesSubsystem.cpp
+++ b/Engine/Plugins/Runtime/GameFeatures/Source/GameFeatures/Private/GameFeaturesSubsystem.cpp
@@ -3,6 +3,7 @@
 #include "GameFeaturesSubsystem.h"
 #include "Algo/Accumulate.h"
 #include "Algo/TopologicalSort.h"
+#include "Algo/Unique.h"
 #include "AssetRegistry/AssetData.h"
 #include "AssetRegistry/IAssetRegistry.h"
 #include "BundlePrereqCombinedStatusHelper.h"
@@ -3454,6 +3455,8 @@ bool UGameFeaturesSubsystem::FindOrCreatePluginDependencyStateMachines(const FSt
 			OutDependencyMachines.Add(ResolvedDependency);
 		}
 
+		OutDependencyMachines.Sort();
+		OutDependencyMachines.SetNum(Algo::Unique(OutDependencyMachines));
 		return true;
 	}
 
@@ -3462,7 +3465,7 @@ bool UGameFeaturesSubsystem::FindOrCreatePluginDependencyStateMachines(const FSt
 
 bool UGameFeaturesSubsystem::FindPluginDependencyStateMachinesToActivate(const FString& PluginURL, const FString& PluginFilename, TArray<UGameFeaturePluginStateMachine*>& OutDependencyMachines) const
 {
-	return EnumeratePluginDependenciesWithShouldActivate(PluginURL, PluginFilename, [this, &OutDependencyMachines](const FString& DependencyName, const FString& DependencyURL) {
+	const bool bResult = EnumeratePluginDependenciesWithShouldActivate(PluginURL, PluginFilename, [this, &OutDependencyMachines](const FString& DependencyName, const FString& DependencyURL) {
 		UGameFeaturePluginStateMachine* Dependency = FindGameFeaturePluginStateMachine(DependencyURL);
 		if (Dependency)
 		{
@@ -3476,11 +3479,15 @@ bool UGameFeaturesSubsystem::FindPluginDependencyStateMachinesToActivate(const F
 			return false;
 		}
 	});
+	OutDependencyMachines.Sort();
+	OutDependencyMachines.SetNum(Algo::Unique(OutDependencyMachines));
+
+	return bResult;
 }
 
 bool UGameFeaturesSubsystem::FindPluginDependencyStateMachinesToDeactivate(const FString& PluginURL, const FString& PluginFilename, TArray<UGameFeaturePluginStateMachine*>& OutDependencyMachines) const
 {
-	return EnumeratePluginDependenciesWithShouldActivate(PluginURL, PluginFilename, [this, &OutDependencyMachines](const FString& DependencyName, const FString& DependencyURL) {
+	const bool bResult = EnumeratePluginDependenciesWithShouldActivate(PluginURL, PluginFilename, [this, &OutDependencyMachines](const FString& DependencyName, const FString& DependencyURL) {
 		UGameFeaturePluginStateMachine* Dependency = FindGameFeaturePluginStateMachine(DependencyURL);
 		if (Dependency)
 		{
@@ -3493,6 +3500,10 @@ bool UGameFeaturesSubsystem::FindPluginDependencyStateMachinesToDeactivate(const
 		}
 		return true;
 	});
+	OutDependencyMachines.Sort();
+	OutDependencyMachines.SetNum(Algo::Unique(OutDependencyMachines));
+
+	return bResult;
 }
 
 template <typename CallableT>
