[GameFeatures] Add a mechanism to allow batch processing specific states within the GFP state machine.

Rationale:
Batch processing opens up optimization oppportunities to further improve GFP perf.
Batch processing was enabled on `FGameFeaturePluginState_Activating` to batch apply plugin ini files and showed gains of around 4-5s to map load times.

Implementation goals:
1. Should be opt in so as to minimize risk/fallout
2. Should be reasonably generic so it could be applied/extended to different parts of the GFP state machine.
3. Ideally does not drastically change existing behavior/order of operations (apart from just the batch processing aspect).

Implementation details:
1. Batch processing can be enabled through protocol options.
2. If the option is enabled on a state machine, when entering a batchable state, we queue a batching request.
3. As other state machines hit the batched state, GamePluginStateMachineManager ticks to check if there are no active transitions below the fenced state.
4. Batch processing can be skipped per plugin instance if needed by overriding `CanBatchProcess`.

Batch processing can be disabled through a CVar (`GameFeaturePlugin.EnableBatchProcessing`).

#tests Batch activation through matchmaking runs and RRs. Tested that error states are handled appropriately by manually triggering errors on certain sates. Executable override RRs for measuring gains.
#rnx

diff --git a/Engine/Plugins/Runtime/GameFeatures/Source/GameFeatures/Private/GameFeatureData.cpp b/Engine/Plugins/Runtime/GameFeatures/Source/GameFeatures/Private/GameFeatureData.cpp
--- a/Engine/Plugins/Runtime/GameFeatures/Source/GameFeatures/Private/GameFeatureData.cpp
+++ b/Engine/Plugins/Runtime/GameFeatures/Source/GameFeatures/Private/GameFeatureData.cpp
@@ -1,6 +1,7 @@
 // Copyright Epic Games, Inc. All Rights Reserved.
 
 #include "GameFeatureData.h"
+#include "Algo/Accumulate.h"
 #include "AssetRegistry/AssetData.h"
 #include "Engine/AssetManager.h"
 #include "GameFeaturesSubsystem.h"
@@ -172,7 +173,7 @@ void UGameFeatureData::InitializeBasePluginIniFile(const FString& PluginInstalle
 	}
 }
 
-void UGameFeatureData::InitializeHierarchicalPluginIniFiles(const FString& PluginInstalledFilename) const
+void UGameFeatureData::InitializeHierarchicalPluginIniFiles(const FString& PluginInstalledFilename)
 {
 	static bool bUseNewDynamicLayers = IConsoleManager::Get().FindConsoleVariable(TEXT("ini.UseNewDynamicLayers"))->GetInt() != 0;
 	if (bUseNewDynamicLayers)
@@ -578,6 +579,64 @@ void UGameFeatureData::InitializeHierarchicalPluginIniFiles(const FString& Plugi
 	}
 }
 
+void UGameFeatureData::InitializeHierarchicalPluginIniFiles(const TArrayView<FString>& PluginInstalledFilenames)
+{
+	static bool bUseNewDynamicLayers = IConsoleManager::Get().FindConsoleVariable(TEXT("ini.UseNewDynamicLayers"))->GetInt() != 0;
+	if (bUseNewDynamicLayers)
+	{
+		TArray<FName> PluginNames;
+		PluginNames.Reserve(PluginInstalledFilenames.Num());
+		for (const FString& PluginInstalledFilename : PluginInstalledFilenames)
+		{
+			PluginNames.Emplace(*FPaths::GetBaseFilename(PluginInstalledFilename));
+		}
+
+		if (UE_GET_LOG_VERBOSITY(LogGameFeatures) >= ELogVerbosity::Verbose)
+		{
+			FString PluginList = FString::Printf(TEXT("[%s]"), *Algo::Accumulate(PluginInstalledFilenames, FString(), [](FString InResult, const FString& InName)
+			{
+				InResult = InResult.IsEmpty() ? InName : InResult + "," + InName;
+				return InResult;
+			}));
+			UE_LOG(LogGameFeatures, Verbose, TEXT("Loading GameFeature config modification for %s"), *PluginList);
+		}
+
+		UE::DynamicConfig::PerformDynamicConfig("InitializeHierarchicalPluginIniFiles", [&PluginNames](FConfigModificationTracker* ChangeTracker)
+			{
+				// set which sections to track for cvars, with their priority
+				ChangeTracker->CVars.Add(TEXT("ConsoleVariables")).CVarPriority = (int)ECVF_SetByPluginLowPriority;
+				ChangeTracker->CVars.Add(TEXT("ConsoleVariables_HighPriority")).CVarPriority = (int)ECVF_SetByPluginHighPriority;
+
+				// apply plugin modifications from this plugin to the everything
+				FConfigCacheIni::AddMultiplePluginsToAllBranches(PluginNames, ChangeTracker);
+			});
+
+		UE::DynamicConfig::PerformDynamicConfig("InitializeHierarchicalPluginIniFiles", [&PluginNames](FConfigModificationTracker* ChangeTracker)
+			{
+				{
+					// set which sections to track for cvars, with their priority
+					ChangeTracker->CVars.Add(TEXT("ConsoleVariables")).CVarPriority = (int)ECVF_SetByHotfix;
+					ChangeTracker->CVars.Add(TEXT("ConsoleVariables_HighPriority")).CVarPriority = (int)ECVF_SetByHotfix;
+
+					for (const FName PluginName : PluginNames)
+					{
+						// give hotfix a chance to run at that point
+#define HOTFIX_BRANCH(Ini) UE::DynamicConfig::HotfixPluginForBranch.Broadcast(PluginName, #Ini, ChangeTracker);
+						ENUMERATE_KNOWN_INI_FILES(HOTFIX_BRANCH);
+#undef HOTFIX_BRANCH
+					}
+				}
+			});
+	}
+	else
+	{
+		for (const FString& PluginInstalledFilename : PluginInstalledFilenames)
+		{
+			InitializeHierarchicalPluginIniFiles(PluginInstalledFilename);
+		}
+	}
+}
+
 void UGameFeatureData::ReloadConfigs(FConfigFile& PluginConfig)
 {
 	// Reload configs so objects get the changes
diff --git a/Engine/Plugins/Runtime/GameFeatures/Source/GameFeatures/Private/GameFeaturePluginStateMachine.cpp b/Engine/Plugins/Runtime/GameFeatures/Source/GameFeatures/Private/GameFeaturePluginStateMachine.cpp
--- a/Engine/Plugins/Runtime/GameFeatures/Source/GameFeatures/Private/GameFeaturePluginStateMachine.cpp
+++ b/Engine/Plugins/Runtime/GameFeatures/Source/GameFeatures/Private/GameFeaturePluginStateMachine.cpp
@@ -108,6 +108,10 @@ namespace UE::GameFeatures
 		true,
 		TEXT("Enable experimental asset streaming"));
 
+	static TAutoConsoleVariable<bool> CVarEnableBatchProcessing(TEXT("GameFeaturePlugin.EnableBatchProcessing"),
+		false,
+		TEXT("Enable batch processing when processing multiple plugins and specified in protocol options."));
+
 	bool ShouldDeferLocalizationDataLoad()
 	{
 		// Note: We don't defer localization data loading in the editor, as the editor only needs to mount plugins to use them
@@ -553,6 +557,23 @@ void FGameFeaturePluginState::UpdateProgress(float Progress) const
 	StateProperties.OnFeatureStateProgressUpdate.ExecuteIfBound(Progress);
 }
 
+bool FGameFeaturePluginState::CanBatchProcess() const
+{
+	// Batch processing is tick driven so is technically "async". 
+	// Hence if we are sync loading, avoid using batch processing since it could impact order of operations.
+	return UseAsyncLoading();
+}
+
+bool FGameFeaturePluginState::IsWaitingForBatchProcessing() const
+{
+	return StateProperties.IsWaitingForBatchProcessing();
+}
+
+bool FGameFeaturePluginState::WasBatchProcessed() const
+{
+	return StateProperties.WasBatchProcessed();
+}
+
 void FGameFeaturePluginState::CleanupDeferredUpdateCallbacks() const
 {
 	if (TickHandle.IsValid())
@@ -733,6 +754,11 @@ struct FTransitionDependenciesGameFeaturePluginState : public FGameFeaturePlugin
 	void TransitionDependency(UGameFeaturePluginStateMachine* Dependency)
 	{
 		bool bSetDestination = false;
+
+		if (TransitionPolicy::ExcludeDepedenciesFromBatchProcessing())
+		{
+			Dependency->ExcludeFromBatchProcessing();
+		}
 		
 		if (TransitionPolicy::ShouldWaitForDependencies())
 		{
@@ -2622,6 +2648,11 @@ struct FWaitingForDependenciesTransitionPolicy
 		return EGameFeaturePluginState::ErrorWaitingForDependencies;
 	}
 
+	static bool ExcludeDepedenciesFromBatchProcessing()
+	{
+		return false;
+	}
+
 	static bool ShouldWaitForDependencies()
 	{
 		return true;
@@ -3406,6 +3437,11 @@ struct FDeactivatingDependenciesTransitionPolicy
 		return EGameFeaturePluginState::ErrorDeactivatingDependencies;
 	}
 
+	static bool ExcludeDepedenciesFromBatchProcessing()
+	{
+		return false;
+	}
+
 	static bool ShouldWaitForDependencies()
 	{
 		return UE::GameFeatures::CVarWaitForDependencyDeactivation.GetValueOnGameThread();
@@ -3562,6 +3598,11 @@ struct FActivatingDependenciesTransitionPolicy
 		return EGameFeaturePluginState::ErrorActivatingDependencies;
 	}
 
+	static bool ExcludeDepedenciesFromBatchProcessing()
+	{
+		return true;
+	}
+
 	static bool ShouldWaitForDependencies()
 	{
 		return true;
@@ -3580,22 +3621,34 @@ struct FGameFeaturePluginState_Activating : public FGameFeaturePluginState
 {
 	FGameFeaturePluginState_Activating(FGameFeaturePluginStateMachineProperties& InStateProperties) : FGameFeaturePluginState(InStateProperties) {}
 
+	bool CanBatchProcess() const override
+	{
+		return FGameFeaturePluginState::CanBatchProcess() && AllowIniLoading();
+	}
+
 	virtual void UpdateState(FGameFeaturePluginStateStatus& StateStatus) override
 	{
 		TRACE_CPUPROFILER_EVENT_SCOPE(GFP_Activating);
 		check(GEngine);
 		check(StateProperties.GameFeatureData);
 
-		FGameFeatureActivatingContext Context;
+		if (IsWaitingForBatchProcessing())
+		{
+			return;
+		}
 
-		if (AllowIniLoading())
+		if(!WasBatchProcessed())
 		{
-			TRACE_CPUPROFILER_EVENT_SCOPE(GFP_Activating_InitIni);
-			StateProperties.GameFeatureData->InitializeHierarchicalPluginIniFiles(StateProperties.PluginInstalledFilename);
+			if (AllowIniLoading())
+			{
+				TRACE_CPUPROFILER_EVENT_SCOPE(GFP_Activating_InitIni);
+				StateProperties.GameFeatureData->InitializeHierarchicalPluginIniFiles(StateProperties.PluginInstalledFilename);
+			}
 		}
 
 		{
 			TRACE_CPUPROFILER_EVENT_SCOPE(GFP_Activating_SendEvents);
+			FGameFeatureActivatingContext Context;
 			UGameFeaturesSubsystem::Get().OnGameFeatureActivating(StateProperties.GameFeatureData, StateProperties.PluginName, Context, StateProperties.PluginIdentifier);
 		}
 
@@ -3610,6 +3663,25 @@ struct FGameFeaturePluginState_Activating : public FGameFeaturePluginState
 
 		StateStatus.SetTransition(EGameFeaturePluginState::Active);
 	}
+
+	static void BatchProcess(TConstArrayView<UGameFeaturePluginStateMachine*> GFPs)
+	{
+		TRACE_CPUPROFILER_EVENT_SCOPE(GFP_BatchProcess_OnFenceCompleteActivating);
+		UGameFeaturesSubsystem& GFPSubSys = UGameFeaturesSubsystem::Get();
+
+		TArray<FString> PluginInstalledFilenames;
+		PluginInstalledFilenames.Reserve(GFPs.Num());
+
+		for (const UGameFeaturePluginStateMachine* GFPSM : GFPs)
+		{
+			PluginInstalledFilenames.Emplace(GFPSM->GetProperties().PluginInstalledFilename);
+		}
+
+		{
+			TRACE_CPUPROFILER_EVENT_SCOPE(GFP_BatchActivating_InitIni);
+			UGameFeatureData::InitializeHierarchicalPluginIniFiles(PluginInstalledFilenames);
+		}
+	}
 };
 
 struct FGameFeaturePluginState_Active : public FDestinationGameFeaturePluginState
@@ -3641,11 +3713,64 @@ struct FGameFeaturePluginState_Active : public FDestinationGameFeaturePluginStat
   State Machine
 =========================================================
 */
+namespace UE::GameFeatures
+{
+	/** Helper classes that look for existence of batch processing functions on states so we can identify these compile time */
+	template<class T, class = void>
+	struct TBatchProcessHelperWrapper
+	{
+		static bool ImplementsBatchProcess()
+		{
+			return false;
+		}
+		static void BatchProcess(TConstArrayView<UGameFeaturePluginStateMachine*>)
+		{
+			check(!"Not implemented");
+		}
+	};
+
+	template<class T>
+	struct TBatchProcessHelperWrapper<T, std::enable_if_t<std::is_invocable_r<void, decltype(T::BatchProcess), TConstArrayView<UGameFeaturePluginStateMachine*>>::value>>
+	{
+		static bool ImplementsBatchProcess()
+		{
+			return true;
+		}
+
+		static void BatchProcess(TConstArrayView<UGameFeaturePluginStateMachine*> GFPSMs)
+		{
+			T::BatchProcess(GFPSMs);
+		}
+	};
+
+	struct FBatchProcessHelperFunctors
+	{
+		using FImplementsBatchProcessPtr = bool(*)();
+		using FBatchProcessPtr = void (*)(TConstArrayView<UGameFeaturePluginStateMachine*>);
+
+		FBatchProcessHelperFunctors(FImplementsBatchProcessPtr InImplementsBatchProcessPtr, FBatchProcessPtr InBatchProcessPtr)
+			: ImplementsBatchProcess(InImplementsBatchProcessPtr)
+			, BatchProcess(InBatchProcessPtr)
+		{
+
+		}
+		
+		FImplementsBatchProcessPtr ImplementsBatchProcess = nullptr;
+		FBatchProcessPtr BatchProcess = nullptr;
+	};
+
+	#define GAME_FEATURE_PLUGIN_STATE_MAKE_BATCH_PROCESS_FN(inEnum, inText) FBatchProcessHelperFunctors(&TBatchProcessHelperWrapper<FGameFeaturePluginState_##inEnum>::ImplementsBatchProcess, &TBatchProcessHelperWrapper<FGameFeaturePluginState_##inEnum>::BatchProcess), 
+	static TStaticArray<FBatchProcessHelperFunctors, EGameFeaturePluginState::MAX> BatchProcessingHelperFunctors = {
+		GAME_FEATURE_PLUGIN_STATE_LIST(GAME_FEATURE_PLUGIN_STATE_MAKE_BATCH_PROCESS_FN)
+	};
+	#undef GAME_FEATURE_PLUGIN_STATE_MAKE_BATCH_PROCESS_FN
+}
 
 UGameFeaturePluginStateMachine::UGameFeaturePluginStateMachine(const FObjectInitializer& ObjectInitializer)
 	: Super(ObjectInitializer)
 	, CurrentStateInfo(EGameFeaturePluginState::Uninitialized)
 	, bInUpdateStateMachine(false)
+	, bRegisteredAsTransitioningGFPSM(false)
 {
 
 }
@@ -3666,6 +3791,7 @@ void UGameFeaturePluginStateMachine::InitStateMachine(FGameFeaturePluginIdentifi
 	GAME_FEATURE_PLUGIN_STATE_LIST(GAME_FEATURE_PLUGIN_STATE_MAKE_STATE)
 #undef GAME_FEATURE_PLUGIN_STATE_MAKE_STATE
 
+	CheckAddBatchingRequestForCurrentState();
 	AllStates[CurrentStateInfo.State]->BeginState();
 }
 
@@ -4071,6 +4197,8 @@ void UGameFeaturePluginStateMachine::UpdateStateMachine()
 		DoCallbacks(Result, InState, InState + 1);
 	};
 
+	RegisterAsTransitioningStateMachine();
+
 	bool bKeepProcessing = false;
 	int32 NumTransitions = 0;
 	const int32 MaxTransitions = 10000;
@@ -4095,12 +4223,14 @@ void UGameFeaturePluginStateMachine::UpdateStateMachine()
 			{
 				TRACE_CPUPROFILER_EVENT_SCOPE(GFP_UpdateStateMachine_EndState);
 				AllStates[CurrentState]->EndState();
+				CheckAndCancelBatchingRequestForCurrentState();
 			}
 			CurrentStateInfo = FGameFeaturePluginStateInfo(StateStatus.TransitionToState);
 			CurrentState = StateStatus.TransitionToState;
 			check(CurrentState != EGameFeaturePluginState::MAX);
 			{
 				TRACE_CPUPROFILER_EVENT_SCOPE(GFP_UpdateStateMachine_BeginState);
+				CheckAddBatchingRequestForCurrentState();
 				AllStates[CurrentState]->BeginState();
 			}
 
@@ -4156,16 +4286,27 @@ void UGameFeaturePluginStateMachine::UpdateStateMachine()
 				DoCallback(StateStatus.TransitionResult, CurrentState);
 			}
 
+			if (!bKeepProcessing)
+			{
+				UnregisterAsTransitioningStateMachine();
+			}
+
 			if (CurrentState == EGameFeaturePluginState::Terminal)
 			{
 				TRACE_CPUPROFILER_EVENT_SCOPE(GFP_UpdateStateMachine_FinishTerm);
 
+				UnregisterAsTransitioningStateMachine();
+
 				check(bKeepProcessing == false);
 				// Now that callbacks are done this machine can be cleaned up
 				UGameFeaturesSubsystem::Get().FinishTermination(this);
 				MarkAsGarbage();
 			}
 		}
+		else if(!IsRunning())
+		{
+			UnregisterAsTransitioningStateMachine();
+		}
 
 		// Log our final state if we've finished transitioning
 		if (!bKeepProcessing && InitialState != CurrentState)
@@ -4201,6 +4342,91 @@ void UGameFeaturePluginStateMachine::UpdateCurrentStateProgress(float Progress)
 	CurrentStateInfo.Progress = Progress;
 }
 
+void UGameFeaturePluginStateMachine::RegisterAsTransitioningStateMachine()
+{
+	if (bRegisteredAsTransitioningGFPSM)
+	{
+		return;
+	}
+
+	UGameFeaturesSubsystem::Get().RegisterRunningStateMachine(this);
+	bRegisteredAsTransitioningGFPSM = true;
+}
+
+void UGameFeaturePluginStateMachine::UnregisterAsTransitioningStateMachine()
+{
+	if (!bRegisteredAsTransitioningGFPSM)
+	{
+		return;
+	}
+
+	UGameFeaturesSubsystem::Get().UnregisterRunningStateMachine(this);
+	bRegisteredAsTransitioningGFPSM = false;
+}
+
+void UGameFeaturePluginStateMachine::CheckAddBatchingRequestForCurrentState()
+{
+	check(!StateProperties.BatchProcessingHandle.IsValid());
+
+	const bool bCanBatchProcess =
+		UE::GameFeatures::BatchProcessingHelperFunctors[CurrentStateInfo.State].ImplementsBatchProcess() &&
+		StateProperties.CanBatchProcess() &&
+		AllStates[CurrentStateInfo.State]->CanBatchProcess();
+	
+	if (bCanBatchProcess)
+	{
+		UE_LOG(LogGameFeatures, Verbose, TEXT("Game feature '%s' awaiting batch processing of state (%s)"), *GetGameFeatureName(), *UE::GameFeatures::ToString(CurrentStateInfo.State));
+		StateProperties.BatchProcessingHandle = UGameFeaturesSubsystem::Get().AddBatchingRequest(
+			CurrentStateInfo.State,
+			StateProperties.OnRequestUpdateStateMachine);
+	}
+}
+
+void UGameFeaturePluginStateMachine::CheckAndCancelBatchingRequestForCurrentState()
+{
+	// Reset batch processing state.
+	StateProperties.bWasBatchProcessed = false;
+
+	// If we are currently awaiting batch processing, cancel and update the state machine.
+	if (StateProperties.IsWaitingForBatchProcessing())
+	{
+		UE_LOG(LogGameFeatures, Verbose, TEXT("Game feature '%s' cancelled batch processing of state (%s)"), *GetGameFeatureName(), *UE::GameFeatures::ToString(CurrentStateInfo.State));
+		UGameFeaturesSubsystem::Get().CancelBatchingRequest(
+			CurrentStateInfo.State,
+			StateProperties.BatchProcessingHandle);
+
+		StateProperties.BatchProcessingHandle.Reset();
+		UpdateStateMachine();
+	}
+}
+
+/* static */ void UGameFeaturePluginStateMachine::BatchProcess(EGameFeaturePluginState State, TConstArrayView<UGameFeaturePluginStateMachine*> GFPSMs)
+{
+	UE::GameFeatures::BatchProcessingHelperFunctors[State].BatchProcess(GFPSMs);
+	for (UGameFeaturePluginStateMachine* GFPSM : GFPSMs)
+	{
+		GFPSM->StateProperties.BatchProcessingHandle.Reset();
+		GFPSM->StateProperties.bWasBatchProcessed = true;
+	}
+}
+
+void UGameFeaturePluginStateMachine::ExcludeFromBatchProcessing()
+{
+	// Ensure protocol options are updated to reflect the exclusion of this state machine from batch processing.
+	if (StateProperties.ProtocolOptions.bBatchProcess)
+	{
+		UE_LOG(LogGameFeatures, Verbose, TEXT("Game feature '%s' excluded from batch processing"), *GetGameFeatureName());
+
+		FGameFeatureProtocolOptions NewOptions = StateProperties.ProtocolOptions;
+		NewOptions.bBatchProcess = false;
+
+		bool bOutDidUpdate = false;
+		TryUpdatePluginProtocolOptions(NewOptions, bOutDidUpdate);
+		check(bOutDidUpdate);
+		CheckAndCancelBatchingRequestForCurrentState();
+	}	
+}
+
 FGameFeaturePluginStateMachineProperties::FGameFeaturePluginStateMachineProperties(
 	FGameFeaturePluginIdentifier InPluginIdentifier,
 	const FGameFeaturePluginStateRange& DesiredDestination,
@@ -4389,4 +4615,19 @@ bool FGameFeaturePluginStateMachineProperties::AllowAsyncLoading() const
 		(!IsRunningCommandlet() || UE::GameFeatures::CVarForceAsyncLoad.GetValueOnGameThread());
 }
 
+bool FGameFeaturePluginStateMachineProperties::CanBatchProcess() const
+{
+	return ProtocolOptions.bBatchProcess && UE::GameFeatures::CVarEnableBatchProcessing.GetValueOnGameThread();
+}
+
+bool FGameFeaturePluginStateMachineProperties::IsWaitingForBatchProcessing() const
+{
+	return BatchProcessingHandle.IsValid();
+}
+
+bool FGameFeaturePluginStateMachineProperties::WasBatchProcessed() const
+{
+	return bWasBatchProcessed;
+}
+
 #undef LOCTEXT_NAMESPACE 
diff --git a/Engine/Plugins/Runtime/GameFeatures/Source/GameFeatures/Private/GameFeaturePluginStateMachine.h b/Engine/Plugins/Runtime/GameFeatures/Source/GameFeatures/Private/GameFeaturePluginStateMachine.h
--- a/Engine/Plugins/Runtime/GameFeatures/Source/GameFeatures/Private/GameFeaturePluginStateMachine.h
+++ b/Engine/Plugins/Runtime/GameFeatures/Source/GameFeatures/Private/GameFeaturePluginStateMachine.h
@@ -260,9 +260,6 @@ DECLARE_DELEGATE_OneParam(FGameFeatureStateTransitionCanceled, UGameFeaturePlugi
 /** A request for other state machine dependencies */
 DECLARE_DELEGATE_RetVal_TwoParams(bool, FGameFeaturePluginRequestStateMachineDependencies, const FString& /*DependencyPluginURL*/, TArray<UGameFeaturePluginStateMachine*>& /*OutDependencyMachines*/);
 
-/** A request to update the state machine and process states */
-DECLARE_DELEGATE(FGameFeaturePluginRequestUpdateStateMachine);
-
 /** A request to update progress for the current state */
 DECLARE_DELEGATE_OneParam(FGameFeatureStateProgressUpdate, float Progress);
 
@@ -305,6 +302,12 @@ struct FGameFeaturePluginStateMachineProperties
 
 	TArray<FName> AddedPrimaryAssetTypes;
 
+	/** Batch Processing handle if requested for batch processed */
+	FDelegateHandle BatchProcessingHandle;
+
+	/** Tracks if the current state was batch processed */
+	bool bWasBatchProcessed : 1 = false;
+
 	/** The data asset describing this game feature */
 	UPROPERTY(Transient)
 	TObjectPtr<UGameFeatureData> GameFeatureData = nullptr;
@@ -338,6 +341,10 @@ struct FGameFeaturePluginStateMachineProperties
 
 	/** Whether this machine is allowed to be asynchronous */
 	bool AllowAsyncLoading() const;
+
+	bool CanBatchProcess() const;
+	bool IsWaitingForBatchProcessing() const;
+	bool WasBatchProcessed() const;
 };
 
 /** Input and output information for a state's UpdateState */
@@ -403,6 +410,10 @@ struct FGameFeaturePluginState
 
 	void UpdateProgress(float Progress) const;
 
+	virtual bool CanBatchProcess() const;
+	bool IsWaitingForBatchProcessing() const;
+	bool WasBatchProcessed() const;
+
 protected:
 	/** Builds an end FResult with some minimal error information with overrides for common types we 
 		need to generate errors from */
@@ -563,6 +574,12 @@ public:
 	/** Sets that the state machine is in an unrecoverable error. Used when the GFD is missing as that can't be recovered from. */
 	void SetUnrecoverableError();
 
+	/** Request to exclude this plugin from batch processing, for e.g. if the plugin is a dependency and has batch processing enabled */
+	void ExcludeFromBatchProcessing();
+
+	/** Helper static function to trigger state specific batch processing functions */
+	static void BatchProcess(EGameFeaturePluginState State, TConstArrayView<UGameFeaturePluginStateMachine*> GFPSMs);	
+
 private:
 	/** Returns true if the specified state is not a transition state */
 	bool IsValidTransitionState(EGameFeaturePluginState InState) const;
@@ -579,6 +596,18 @@ private:
 	/** Update Progress for current state */
 	void UpdateCurrentStateProgress(float Progress);
 
+	/** Helper to register the state machine with the manager as a state machine in transiton */
+	void RegisterAsTransitioningStateMachine();
+
+	/** Helper to unregister the state machine with the manager as a state machine in transiton */
+	void UnregisterAsTransitioningStateMachine();
+	
+	/** Adds a batching request for current state if state can be batched */
+	void CheckAddBatchingRequestForCurrentState();
+
+	/** Cancels batching re1quest for current state if one exists */
+	void CheckAndCancelBatchingRequestForCurrentState();
+
 	/** Information about the current state */
 	FGameFeaturePluginStateInfo CurrentStateInfo;
 
@@ -594,4 +623,7 @@ private:
 
 	/** True when the state machine can not transition out of the current error state or if requested to would likely fail */
 	bool bIsInUnrecoverableError = false;
+
+	/** True when we are registered as a state machine with in flight transitions */
+	bool bRegisteredAsTransitioningGFPSM;
 };
diff --git a/Engine/Plugins/Runtime/GameFeatures/Source/GameFeatures/Private/GameFeaturesSubsystem.cpp b/Engine/Plugins/Runtime/GameFeatures/Source/GameFeatures/Private/GameFeaturesSubsystem.cpp
--- a/Engine/Plugins/Runtime/GameFeatures/Source/GameFeatures/Private/GameFeaturesSubsystem.cpp
+++ b/Engine/Plugins/Runtime/GameFeatures/Source/GameFeatures/Private/GameFeaturesSubsystem.cpp
@@ -325,6 +325,7 @@ bool FInstallBundlePluginProtocolOptions::operator==(const FInstallBundlePluginP
 
 FGameFeatureProtocolOptions::FGameFeatureProtocolOptions()
 	: bForceSyncLoading(false)
+	, bBatchProcess(false)
 	, bLogWarningOnForcedDependencyCreation(false)
 	, bLogErrorOnForcedDependencyCreation(false)
 { 
@@ -341,6 +342,7 @@ FGameFeatureProtocolOptions::FGameFeatureProtocolOptions(const FInstallBundlePlu
 
 FGameFeatureProtocolOptions::FGameFeatureProtocolOptions(FNull InOptions)
 	: bForceSyncLoading(false)
+	, bBatchProcess(false)
 	, bLogWarningOnForcedDependencyCreation(false)
 	, bLogErrorOnForcedDependencyCreation(false)
 {
@@ -479,6 +481,12 @@ void UGameFeaturesSubsystem::Deinitialize()
 	}
 	GameSpecificPolicies = nullptr;
 	bInitializedPolicyManager = false;
+
+	if (TickHandle.IsValid())
+	{
+		FTSTicker::GetCoreTicker().RemoveTicker(MoveTemp(TickHandle));
+		TickHandle.Reset();
+	}
 }
 
 void UGameFeaturesSubsystem::OnAssetManagerCreated()
@@ -1982,6 +1990,7 @@ void UGameFeaturesSubsystem::LoadBuiltInGameFeaturePlugin(const TSharedRef<IPlug
 				ProtocolOptions.bForceSyncLoading = BehaviorOptions.bForceSyncLoading;
 				ProtocolOptions.bLogWarningOnForcedDependencyCreation = BehaviorOptions.bLogWarningOnForcedDependencyCreation;
 				ProtocolOptions.bLogErrorOnForcedDependencyCreation = BehaviorOptions.bLogErrorOnForcedDependencyCreation;
+				ProtocolOptions.bBatchProcess = BehaviorOptions.bBatchProcess;
 				UGameFeaturePluginStateMachine* StateMachine = FindOrCreateGameFeaturePluginStateMachine(PluginURL, ProtocolOptions);
 
 				EBuiltInAutoState InitialAutoState = (BehaviorOptions.AutoStateOverride != EBuiltInAutoState::Invalid) ? 
@@ -3785,6 +3794,127 @@ void UGameFeaturesSubsystem::CallbackObservers(EObserverCallback CallbackType, c
 	}
 }
 
+void UGameFeaturesSubsystem::RegisterRunningStateMachine(UGameFeaturePluginStateMachine* GFPSM)
+{
+	check(!RunningStateMachines.Contains(GFPSM));
+	RunningStateMachines.Add(GFPSM);
+}
+
+void UGameFeaturesSubsystem::UnregisterRunningStateMachine(UGameFeaturePluginStateMachine* GFPSM)
+{
+	verify(RunningStateMachines.Remove(GFPSM) == 1);
+}
+
+FDelegateHandle UGameFeaturesSubsystem::AddBatchingRequest(EGameFeaturePluginState State, FGameFeaturePluginRequestUpdateStateMachine UpdateDelegate)
+{
+	// Adding first?
+	if (BatchProcessingFences.IsEmpty())
+	{
+		EnableTick();
+	}
+
+	FGameFeatureBatchProcessingFence& Fence = BatchProcessingFences.FindOrAdd(State);
+	return Fence.NotifyUpdateStateMachines.Add(UpdateDelegate);
+}
+
+void UGameFeaturesSubsystem::CancelBatchingRequest(EGameFeaturePluginState State, FDelegateHandle DelegateHandle)
+{
+	FGameFeatureBatchProcessingFence& Fence = BatchProcessingFences[State];
+	Fence.NotifyUpdateStateMachines.Remove(DelegateHandle);
+}
+
+void UGameFeaturesSubsystem::EnableTick()
+{
+	if (!TickHandle.IsValid())
+	{
+		TickHandle = FTSTicker::GetCoreTicker().AddTicker(FTickerDelegate::CreateUObject(this, &UGameFeaturesSubsystem::Tick));
+		UE_LOG(LogGameFeatures, Verbose, TEXT("UGameFeaturesSubsystem enabled tick [InFlightTransitions:%d]"), RunningStateMachines.Num());
+	}
+}
+
+void UGameFeaturesSubsystem::DisableTick()
+{
+	if (TickHandle.IsValid())
+	{
+		TickHandle.Reset();
+		UE_LOG(LogGameFeatures, Verbose, TEXT("UGameFeaturesSubsystem disabled tick [InFlightTransitions:%d]"), RunningStateMachines.Num());
+	}
+}
+
+bool UGameFeaturesSubsystem::Tick(float DeltaTime)
+{
+	TRACE_CPUPROFILER_EVENT_SCOPE(GFP_GameFeaturesSubsystem_Tick);
+
+	bool bContinueTicking = TickBatchProcessing();
+	if (!bContinueTicking)
+	{
+		DisableTick();
+		return false;
+	}
+
+	return true;
+}
+
+bool UGameFeaturesSubsystem::TickBatchProcessing()
+{
+	TRACE_CPUPROFILER_EVENT_SCOPE(GFP_GameFeaturesSubsystem_TickFences);
+
+	// No fences?
+	if (BatchProcessingFences.IsEmpty())
+	{
+		return false;
+	}
+
+	// Find the closest fence to process.
+	EGameFeaturePluginState CurrentFence = EGameFeaturePluginState::MAX;
+	for (std::underlying_type<EGameFeaturePluginState>::type iState = 0;
+		iState < EGameFeaturePluginState::MAX;
+		++iState)
+	{
+		if (BatchProcessingFences.Contains((EGameFeaturePluginState)iState))
+		{
+			CurrentFence = (EGameFeaturePluginState)(iState);
+			break;
+		}
+	}
+
+	check(CurrentFence != EGameFeaturePluginState::MAX);
+
+	// Check for in flight transitions on state machines to be included in fence
+	TArray<UGameFeaturePluginStateMachine*> AtFence;
+	AtFence.Reserve(RunningStateMachines.Num());
+	bool bAreAnyGFPSMsPendingAtFence = false;
+	for (UGameFeaturePluginStateMachine* GFPSM : RunningStateMachines)
+	{
+		if (GFPSM->GetProperties().IsWaitingForBatchProcessing() && GFPSM->GetCurrentState() == CurrentFence)
+		{
+			AtFence.Add(GFPSM);
+		}
+		if (GFPSM->GetDestination().MinState >= CurrentFence && GFPSM->GetCurrentState() < CurrentFence)
+		{
+			bAreAnyGFPSMsPendingAtFence = true;
+			break;
+		}
+	}
+
+	if (!bAreAnyGFPSMsPendingAtFence)
+	{
+		TRACE_CPUPROFILER_EVENT_SCOPE(GFP_GameFeaturesSubsystem_FenceComplete);
+
+		FGameFeatureBatchProcessingFence& Fence = BatchProcessingFences[CurrentFence];
+
+		UGameFeaturePluginStateMachine::BatchProcess(CurrentFence, AtFence);
+
+		Fence.NotifyUpdateStateMachines.Broadcast();
+
+		BatchProcessingFences.Remove(CurrentFence);
+
+		return BatchProcessingFences.IsEmpty();
+	}
+
+	return true;
+}
+
 TSet<FString> UGameFeaturesSubsystem::GetActivePluginNames() const
 {
 	TSet<FString> ActivePluginNames;
diff --git a/Engine/Plugins/Runtime/GameFeatures/Source/GameFeatures/Public/GameFeatureData.h b/Engine/Plugins/Runtime/GameFeatures/Source/GameFeatures/Public/GameFeatureData.h
--- a/Engine/Plugins/Runtime/GameFeatures/Source/GameFeatures/Public/GameFeatureData.h
+++ b/Engine/Plugins/Runtime/GameFeatures/Source/GameFeatures/Public/GameFeatureData.h
@@ -42,7 +42,8 @@ public:
 	static void InitializeBasePluginIniFile(const FString& PluginInstalledFilename);
 
 	/** Method to process ini files for the plugin during activation */
-	void InitializeHierarchicalPluginIniFiles(const FString& PluginInstalledFilename) const;
+	static void InitializeHierarchicalPluginIniFiles(const FString& PluginInstalledFilename);
+	static void InitializeHierarchicalPluginIniFiles(const TArrayView<FString>& PluginInstalledFilenames);
 
 	UFUNCTION(BlueprintCallable, Category = "GameFeature")
 	static void GetPluginName(const UGameFeatureData* GFD, FString& PluginName);
diff --git a/Engine/Plugins/Runtime/GameFeatures/Source/GameFeatures/Public/GameFeaturesSubsystem.h b/Engine/Plugins/Runtime/GameFeatures/Source/GameFeatures/Public/GameFeaturesSubsystem.h
--- a/Engine/Plugins/Runtime/GameFeatures/Source/GameFeatures/Public/GameFeaturesSubsystem.h
+++ b/Engine/Plugins/Runtime/GameFeatures/Source/GameFeatures/Public/GameFeaturesSubsystem.h
@@ -198,6 +198,10 @@ GAMEFEATURES_API DECLARE_LOG_CATEGORY_EXTERN(LogGameFeatures, Log, All);
 /** Notification that a game feature plugin install/register/load/unload has finished */
 DECLARE_DELEGATE_OneParam(FGameFeaturePluginChangeStateComplete, const UE::GameFeatures::FResult& /*Result*/);
 
+/** A request to update the state machine and process states */
+DECLARE_DELEGATE(FGameFeaturePluginRequestUpdateStateMachine);
+DECLARE_MULTICAST_DELEGATE(FNotifyGameFeaturePluginRequestUpdateStateMachine)
+
 using FGameFeaturePluginLoadComplete = FGameFeaturePluginChangeStateComplete;
 using FGameFeaturePluginDeactivateComplete = FGameFeaturePluginChangeStateComplete;
 using FGameFeaturePluginUnloadComplete = FGameFeaturePluginChangeStateComplete;
@@ -257,6 +261,9 @@ struct FBuiltInGameFeaturePluginBehaviorOptions
 	/** Log Warning if loading this GFP forces dependencies to be created, useful for catching GFP load filtering bugs */
 	bool bLogWarningOnForcedDependencyCreation = false;
 
+	/** Batch process GFPs if/when possible (could be used when processing multiple plugins)*/
+	bool bBatchProcess = false;
+
 	/** Log Error if loading this GFP forces dependencies to be created, useful for catching GFP load filtering bugs */
 	bool bLogErrorOnForcedDependencyCreation = false;
 };
@@ -415,6 +422,7 @@ struct FGameFeatureProtocolOptions : public TUnion<FInstallBundlePluginProtocolO
 	{
 		return TUnion<FInstallBundlePluginProtocolOptions, FNull>::operator==(Other) &&
 			bForceSyncLoading == Other.bForceSyncLoading &&
+			bBatchProcess == Other.bBatchProcess &&
 			bLogWarningOnForcedDependencyCreation == Other.bLogWarningOnForcedDependencyCreation &&
 			bLogErrorOnForcedDependencyCreation == Other.bLogErrorOnForcedDependencyCreation;
 	}
@@ -422,6 +430,9 @@ struct FGameFeatureProtocolOptions : public TUnion<FInstallBundlePluginProtocolO
 	/** Force this GFP to load synchronously even if async loading is allowed */
 	bool bForceSyncLoading : 1;
 
+	/** Batch process GFPs if/when possible (could be used when processing multiple plugins)*/
+	bool bBatchProcess : 1;
+
 	/** Log Warning if loading this GFP forces dependencies to be created, useful for catching GFP load filtering bugs */
 	bool bLogWarningOnForcedDependencyCreation : 1;
 
@@ -872,11 +883,42 @@ private:
 		const UGameFeatureData* GameFeatureData = nullptr, 
 		FGameFeatureStateChangeContext* StateChangeContext = nullptr);
 
+	/** Registers a state machine that is in transition and running */
+	void RegisterRunningStateMachine(UGameFeaturePluginStateMachine* GFPSM);
+
+	/** Unregisters a state machine that was in transition is no longer running */
+	void UnregisterRunningStateMachine(UGameFeaturePluginStateMachine* GFPSM);
+
+	/** Adds a batching request for a given state so we can start listening for when state machines have arrived at a fence */
+	FDelegateHandle AddBatchingRequest(EGameFeaturePluginState State, FGameFeaturePluginRequestUpdateStateMachine UpdateDelegate);
+
+	/** Cancels an existing batching request */
+	void CancelBatchingRequest(EGameFeaturePluginState State, FDelegateHandle DelegateHandle);
+
+	void EnableTick();
+	void DisableTick();
+	bool Tick(float DeltaTime);
+
+	bool TickBatchProcessing();
+
 private:
 	/** The list of all game feature plugin state machine objects */
 	UPROPERTY(Transient)
 	TMap<FString, TObjectPtr<UGameFeaturePluginStateMachine>> GameFeaturePluginStateMachines;
 
+	/** The tick handle if currently registered for a tick */
+	FTSTicker::FDelegateHandle TickHandle;
+
+	/** State machine currently in transition, used to limit search space when checking a batch processing fence or similar */
+	TArray<UGameFeaturePluginStateMachine*> RunningStateMachines;
+
+	/** Active fences */
+	struct FGameFeatureBatchProcessingFence
+	{
+		FNotifyGameFeaturePluginRequestUpdateStateMachine NotifyUpdateStateMachines;
+	};
+	TMap<EGameFeaturePluginState, FGameFeatureBatchProcessingFence> BatchProcessingFences;
+
 	/** Game feature plugin state machine objects that are being terminated. Used to prevent GC until termination is complete. */
 	UPROPERTY(Transient)
 	TArray<TObjectPtr<UGameFeaturePluginStateMachine>> TerminalGameFeaturePluginStateMachines;
