GameFeatureSubsystem
- Add bDoNotDownload option to FBuiltInGameFeaturePluginBehaviorOptions
- Add Plugin_URL_Not_Found not found error code that is seperate from Plugin_Details_Not_Found
- Fix returning error results for plugins that aren't GFPs
- Fix UnmountedPluginNameToExplanation not being populated for some failure cases

diff --git a/Engine/Plugins/Runtime/GameFeatures/Source/GameFeatures/Private/GameFeaturesSubsystem.cpp b/Engine/Plugins/Runtime/GameFeatures/Source/GameFeatures/Private/GameFeaturesSubsystem.cpp
--- a/Engine/Plugins/Runtime/GameFeatures/Source/GameFeatures/Private/GameFeaturesSubsystem.cpp
+++ b/Engine/Plugins/Runtime/GameFeatures/Source/GameFeatures/Private/GameFeaturesSubsystem.cpp
@@ -50,6 +50,7 @@ namespace UE::GameFeatures
 		const TCHAR* PluginNotAllowed = TEXT("Plugin_Denied_By_GameSpecificPolicy");
 		const TCHAR* PluginFiltered = TEXT("Plugin_Filtered_By_GameSpecificPolicy");
 		const TCHAR* PluginDetailsNotFound = TEXT("Plugin_Details_Not_Found");
+		const TCHAR* PluginURLNotFound = TEXT("Plugin_URL_Not_Found");
 		const TCHAR* DependencyFailedRegister = TEXT("Failed_Dependency_Register");
 		const TCHAR* BadURL = TEXT("Bad_URL");
 		const TCHAR* UnreachableState = TEXT("State_Currently_Unreachable");
@@ -1967,67 +1968,99 @@ void UGameFeaturesSubsystem::LoadBuiltInGameFeaturePlugin(const TSharedRef<IPlug
 {
 	UE_SCOPED_ENGINE_ACTIVITY(TEXT("Loading GameFeaturePlugin %s"), *Plugin->GetName());
 
+#if WITH_EDITOR
+	auto AddUnmountedPluginExplination = [this, &Plugin](const UE::GameFeatures::FResult& Result)
+	{
+		UnmountedPluginNameToExplanation.FindOrAdd(Plugin->GetName()) = UE::GameFeatures::ToString(Result);
+		return Result;
+	};
+#else
+	auto AddUnmountedPluginExplination = [](const UE::GameFeatures::FResult& Result)
+	{
+		return Result;
+	};
+#endif // WITH_EDITOR
+
 	UAssetManager::Get().PushBulkScanning();
 
+	const FString& PluginDescriptorFilename = Plugin->GetDescriptorFileName();
+
 	FString PluginURL;
 	FGameFeaturePluginDetails PluginDetails;
-	if (GetBuiltInGameFeaturePluginDetails(Plugin, PluginDetails) && GetBuiltInGameFeaturePluginURL(Plugin, PluginURL))
+
+	if (PluginDescriptorFilename.IsEmpty())
 	{
-		FString PluginDisallowedReason;
-		if (IsPluginAllowed(PluginURL, &PluginDisallowedReason))
+		CompleteDelegate.ExecuteIfBound(UE::GameFeatures::FResult(MakeError(UE::GameFeatures::SubsystemErrorNamespace + UE::GameFeatures::CommonErrorCodes::PluginDetailsNotFound)));
+	}
+	else if (!GetDefault<UGameFeaturesSubsystemSettings>()->IsValidGameFeaturePlugin(FPaths::ConvertRelativePathToFull(PluginDescriptorFilename)))
+	{
+		// Not a GFP, trivial success
+		CompleteDelegate.ExecuteIfBound(UE::GameFeatures::FResult(MakeValue()));
+	}
+	else if (!FPaths::FileExists(PluginDescriptorFilename) || !GetGameFeaturePluginDetailsInternal(PluginDescriptorFilename, PluginDetails))
+	{
+		CompleteDelegate.ExecuteIfBound(AddUnmountedPluginExplination(MakeError(UE::GameFeatures::SubsystemErrorNamespace + UE::GameFeatures::CommonErrorCodes::PluginDetailsNotFound)));
+	}
+	else if (!GetBuiltInGameFeaturePluginURL(Plugin, PluginURL))
+	{
+		CompleteDelegate.ExecuteIfBound(AddUnmountedPluginExplination(MakeError(UE::GameFeatures::SubsystemErrorNamespace + UE::GameFeatures::CommonErrorCodes::PluginURLNotFound)));
+	}
+	else if (FString PluginDisallowedReason; !IsPluginAllowed(PluginURL, &PluginDisallowedReason))
+	{
+		CompleteDelegate.ExecuteIfBound(
+			AddUnmountedPluginExplination(
+				UE::GameFeatures::FResult(
+					MakeError(UE::GameFeatures::SubsystemErrorNamespace + UE::GameFeatures::CommonErrorCodes::PluginNotAllowed),
+					FText::Format(NSLOCTEXT("GameFeatures", "PluginDisallowedText", "Plugin disallowed reason: {0}"), FText::FromString(PluginDisallowedReason))
+				)
+			)
+		);
+	}
+	else
+	{
+		FBuiltInGameFeaturePluginBehaviorOptions BehaviorOptions;
+		const bool bShouldProcess = AdditionalFilter(Plugin->GetDescriptorFileName(), PluginDetails, BehaviorOptions);
+		if (!bShouldProcess)
 		{
-			FBuiltInGameFeaturePluginBehaviorOptions BehaviorOptions;
-			const bool bShouldProcess = AdditionalFilter(Plugin->GetDescriptorFileName(), PluginDetails, BehaviorOptions);
-			if (bShouldProcess)
+			CompleteDelegate.ExecuteIfBound(AddUnmountedPluginExplination(MakeError(UE::GameFeatures::SubsystemErrorNamespace + UE::GameFeatures::CommonErrorCodes::PluginFiltered)));
+		}
+		else
+		{
+			FGameFeatureProtocolOptions ProtocolOptions;
+			if (UGameFeaturesSubsystem::GetPluginURLProtocol(PluginURL) == EGameFeaturePluginProtocol::InstallBundle)
 			{
-				FGameFeatureProtocolOptions ProtocolOptions;
-				if (UGameFeaturesSubsystem::GetPluginURLProtocol(PluginURL) == EGameFeaturePluginProtocol::InstallBundle)
-				{
-					FInstallBundlePluginProtocolOptions InstallBundleOptions;
-					InstallBundleOptions.bAllowIniLoading = true;
-					ProtocolOptions = FGameFeatureProtocolOptions(InstallBundleOptions);
-				}
-				ProtocolOptions.bForceSyncLoading = BehaviorOptions.bForceSyncLoading;
-				ProtocolOptions.bLogWarningOnForcedDependencyCreation = BehaviorOptions.bLogWarningOnForcedDependencyCreation;
-				ProtocolOptions.bLogErrorOnForcedDependencyCreation = BehaviorOptions.bLogErrorOnForcedDependencyCreation;
-				ProtocolOptions.bBatchProcess = BehaviorOptions.bBatchProcess;
-				UGameFeaturePluginStateMachine* StateMachine = FindOrCreateGameFeaturePluginStateMachine(PluginURL, ProtocolOptions);
-
-				EBuiltInAutoState InitialAutoState = (BehaviorOptions.AutoStateOverride != EBuiltInAutoState::Invalid) ? 
-					BehaviorOptions.AutoStateOverride : PluginDetails.BuiltInAutoState;
-				if (InitialAutoState < EBuiltInAutoState::Registered && IsRunningCookCommandlet())
-				{
-					UE_LOG(LogGameFeatures, Display, TEXT("%s will be set to Registerd for cooking"), *Plugin->GetName());
-					InitialAutoState = EBuiltInAutoState::Registered;
-				}
-				const EGameFeaturePluginState DestinationState = ConvertInitialFeatureStateToTargetState(InitialAutoState);
-
-				StateMachine->SetWasLoadedAsBuiltIn();
-
-				// If we're already at the destination or beyond, don't transition back
-				FGameFeaturePluginStateRange Destination(DestinationState, EGameFeaturePluginState::Active);
-				ChangeGameFeatureDestination(StateMachine, ProtocolOptions, Destination,
-					FGameFeaturePluginChangeStateComplete::CreateWeakLambda(this, [this, StateMachine, Destination, CompleteDelegate](const UE::GameFeatures::FResult& Result)
-					{
-						LoadBuiltInGameFeaturePluginComplete(Result, StateMachine, Destination);
-						CompleteDelegate.ExecuteIfBound(Result);
-					}));
+				FInstallBundlePluginProtocolOptions InstallBundleOptions;
+				InstallBundleOptions.bAllowIniLoading = true;
+				InstallBundleOptions.bDoNotDownload = BehaviorOptions.bDoNotDownload;
+				ProtocolOptions = FGameFeatureProtocolOptions(InstallBundleOptions);
 			}
-			else
+			ProtocolOptions.bForceSyncLoading = BehaviorOptions.bForceSyncLoading;
+			ProtocolOptions.bLogWarningOnForcedDependencyCreation = BehaviorOptions.bLogWarningOnForcedDependencyCreation;
+			ProtocolOptions.bLogErrorOnForcedDependencyCreation = BehaviorOptions.bLogErrorOnForcedDependencyCreation;
+			ProtocolOptions.bBatchProcess = BehaviorOptions.bBatchProcess;
+			UGameFeaturePluginStateMachine* StateMachine = FindOrCreateGameFeaturePluginStateMachine(PluginURL, ProtocolOptions);
+
+			EBuiltInAutoState InitialAutoState = (BehaviorOptions.AutoStateOverride != EBuiltInAutoState::Invalid) ? 
+				BehaviorOptions.AutoStateOverride : PluginDetails.BuiltInAutoState;
+			if (InitialAutoState < EBuiltInAutoState::Registered && IsRunningCookCommandlet())
 			{
-				CompleteDelegate.ExecuteIfBound(UE::GameFeatures::FResult(MakeError(UE::GameFeatures::SubsystemErrorNamespace + UE::GameFeatures::CommonErrorCodes::PluginFiltered)));
+				UE_LOG(LogGameFeatures, Display, TEXT("%s will be set to Registerd for cooking"), *Plugin->GetName());
+				InitialAutoState = EBuiltInAutoState::Registered;
 			}
-		}
-		else
-		{
-			CompleteDelegate.ExecuteIfBound(UE::GameFeatures::FResult(MakeError(UE::GameFeatures::SubsystemErrorNamespace + UE::GameFeatures::CommonErrorCodes::PluginNotAllowed),
-				FText::Format(NSLOCTEXT("GameFeatures", "PluginDisallowedText", "Plugin disallowed reason: {0}"), FText::FromString(PluginDisallowedReason))));
+			const EGameFeaturePluginState DestinationState = ConvertInitialFeatureStateToTargetState(InitialAutoState);
+
+			StateMachine->SetWasLoadedAsBuiltIn();
+
+			// If we're already at the destination or beyond, don't transition back
+			FGameFeaturePluginStateRange Destination(DestinationState, EGameFeaturePluginState::Active);
+			ChangeGameFeatureDestination(StateMachine, ProtocolOptions, Destination,
+				FGameFeaturePluginChangeStateComplete::CreateWeakLambda(this, [this, StateMachine, Destination, CompleteDelegate](const UE::GameFeatures::FResult& Result)
+				{
+					LoadBuiltInGameFeaturePluginComplete(Result, StateMachine, Destination);
+					CompleteDelegate.ExecuteIfBound(Result);
+				}));
 		}
 	}
-	else
-	{
-		CompleteDelegate.ExecuteIfBound(UE::GameFeatures::FResult(MakeError(UE::GameFeatures::SubsystemErrorNamespace + UE::GameFeatures::CommonErrorCodes::PluginDetailsNotFound)));
-	}
 
 	UAssetManager::Get().PopBulkScanning();
 }
@@ -2305,9 +2338,6 @@ bool UGameFeaturesSubsystem::GetBuiltInGameFeaturePluginDetails(const TSharedRef
 
 bool UGameFeaturesSubsystem::GetBuiltInGameFeaturePluginDetails(const TSharedRef<IPlugin>& Plugin, struct FGameFeaturePluginDetails& OutPluginDetails) const
 {
-	// @TODO: this problematic because it assumes file protocol.
-	// Ideally this would work with any protocol, but for current uses cases the exact protocol doesn't seem to matter.
-
 	const FString& PluginDescriptorFilename = Plugin->GetDescriptorFileName();
 	// Make sure you are in a game feature plugins folder. All GameFeaturePlugins are rooted in a GameFeatures folder.
 	if (!PluginDescriptorFilename.IsEmpty() && GetDefault<UGameFeaturesSubsystemSettings>()->IsValidGameFeaturePlugin(FPaths::ConvertRelativePathToFull(PluginDescriptorFilename)) && FPaths::FileExists(PluginDescriptorFilename))
diff --git a/Engine/Plugins/Runtime/GameFeatures/Source/GameFeatures/Public/GameFeaturesSubsystem.h b/Engine/Plugins/Runtime/GameFeatures/Source/GameFeatures/Public/GameFeaturesSubsystem.h
--- a/Engine/Plugins/Runtime/GameFeatures/Source/GameFeatures/Public/GameFeaturesSubsystem.h
+++ b/Engine/Plugins/Runtime/GameFeatures/Source/GameFeatures/Public/GameFeaturesSubsystem.h
@@ -258,12 +258,15 @@ struct FBuiltInGameFeaturePluginBehaviorOptions
 	/** Force this GFP to load synchronously even if async loading is allowed */
 	bool bForceSyncLoading = false;
 
-	/** Log Warning if loading this GFP forces dependencies to be created, useful for catching GFP load filtering bugs */
-	bool bLogWarningOnForcedDependencyCreation = false;
-
 	/** Batch process GFPs if/when possible (could be used when processing multiple plugins)*/
 	bool bBatchProcess = false;
 
+	/** Disallows downloading, useful for conditionally loading content only if it's already been installed */
+	bool bDoNotDownload = false;
+
+	/** Log Warning if loading this GFP forces dependencies to be created, useful for catching GFP load filtering bugs */
+	bool bLogWarningOnForcedDependencyCreation = false;
+
 	/** Log Error if loading this GFP forces dependencies to be created, useful for catching GFP load filtering bugs */
 	bool bLogErrorOnForcedDependencyCreation = false;
 };
@@ -406,7 +409,7 @@ struct GAMEFEATURES_API FInstallBundlePluginProtocolOptions
 	/** Allow the GFP to load INI files, should only be allowed for trusted content */
 	bool bAllowIniLoading : 1;
 
-	/** Disallows downloading, useful for conditionally loading content only if it's already been installed **/
+	/** Disallows downloading, useful for conditionally loading content only if it's already been installed */
 	bool bDoNotDownload : 1;
 
 	bool operator==(const FInstallBundlePluginProtocolOptions& Other) const;
