GameFeaturePlugin State Machine
- Fix FTransitionDependenciesGameFeaturePluginState recieving callbacks for dependencies after having Ended.  This can occur when dependency transition failures propogate back up the chain of waiting GFPs because the state may end due to an error while the callback list of its dependency is still being iterated.  This may all happen in the same callstack.

diff --git a/Engine/Plugins/Runtime/GameFeatures/Source/GameFeatures/Private/GameFeaturePluginStateMachine.cpp b/Engine/Plugins/Runtime/GameFeatures/Source/GameFeatures/Private/GameFeaturePluginStateMachine.cpp
--- a/Engine/Plugins/Runtime/GameFeatures/Source/GameFeatures/Private/GameFeaturePluginStateMachine.cpp
+++ b/Engine/Plugins/Runtime/GameFeatures/Source/GameFeatures/Private/GameFeaturePluginStateMachine.cpp
@@ -442,6 +442,59 @@ namespace UE::GameFeatures
 	};
 
 	static FGFPSharedInstallTracker GFPSharedInstallTracker;
+
+	// Callback delegates are moved to the stack before broadcasting.
+	// This class tracks callback delegates on the stack to handle removing callbacks from them
+	// for state machines that are also on the stack.
+	template<typename TCallbackMultiDelegate> 
+		requires 
+			std::is_same_v<TCallbackMultiDelegate, FDestinationGameFeaturePluginState::FOnDestinationStateReached> ||
+			std::is_same_v<TCallbackMultiDelegate, FGameFeaturePluginStateMachineProperties::FOnTransitionCanceled>
+	class TBroadcastingCallback
+	{
+	public:
+		TCallbackMultiDelegate CallbackDelegate;
+
+	private:
+		TBroadcastingCallback<TCallbackMultiDelegate>* Next = nullptr;
+		inline static TBroadcastingCallback<TCallbackMultiDelegate>* Head = nullptr;
+
+	public:
+		TBroadcastingCallback(TCallbackMultiDelegate&& InCallbackDelegate)
+			: CallbackDelegate(MoveTemp(InCallbackDelegate))
+			, Next(Head)
+		{
+			Head = this;
+		}
+
+		~TBroadcastingCallback()
+		{
+			Head = Next;
+		}
+
+		static void RemovePendingCallback(FDelegateHandle InHandle)
+		{
+			TBroadcastingCallback<TCallbackMultiDelegate>* Current = Head;
+			while (Current)
+			{
+				Current->CallbackDelegate.Remove(InHandle);
+				Current = Current->Next;
+			}
+		}
+
+		static void RemovePendingCallback(FDelegateUserObject InObject)
+		{
+			TBroadcastingCallback<TCallbackMultiDelegate>* Current = Head;
+			while (Current)
+			{
+				Current->CallbackDelegate.RemoveAll(InObject);
+				Current = Current->Next;
+			}
+		}
+	};
+
+	using FBroadcastingOnDestinationStateReached = TBroadcastingCallback<FDestinationGameFeaturePluginState::FOnDestinationStateReached>;
+	using FBroadcastingOnTransitionCanceled = TBroadcastingCallback<FGameFeaturePluginStateMachineProperties::FOnTransitionCanceled>;
 }
 
 void FGameFeaturePluginStateStatus::SetTransition(EGameFeaturePluginState InTransitionToState)
@@ -886,17 +939,25 @@ struct FTransitionDependenciesGameFeaturePluginState : public FGameFeaturePlugin
 
 	void ClearDependencies()
 	{
-		for (FDepResultPair& Pair : RemainingDependencies)
+		if (RemainingDependencies.Num() > 0)
 		{
-			UGameFeaturePluginStateMachine* RemainingDependency = Pair.Key.Get();
-			if (RemainingDependency)
+			for (FDepResultPair& Pair : RemainingDependencies)
 			{
-				RemainingDependency->RemovePendingTransitionCallback(this);
-				RemainingDependency->RemovePendingCancelCallback(this);
+				UGameFeaturePluginStateMachine* RemainingDependency = Pair.Key.Get();
+				if (RemainingDependency)
+				{
+					RemainingDependency->RemovePendingTransitionCallback(this);
+					RemainingDependency->RemovePendingCancelCallback(this);
+				}
 			}
+
+			// Also need to cleanup callbacks from any delegates currently on the stack
+			UE::GameFeatures::FBroadcastingOnDestinationStateReached::RemovePendingCallback(this);
+			UE::GameFeatures::FBroadcastingOnTransitionCanceled::RemovePendingCallback(this);
+
+			RemainingDependencies.Empty();
 		}
 
-		RemainingDependencies.Empty();
 		bRequestedDependencies = false;
 	}
 
@@ -4184,10 +4245,10 @@ void UGameFeaturePluginStateMachine::UpdateStateMachine()
 				// Use a local callback on the stack. If SetDestination() is called from the callback then we don't want to stomp the callback
 				// for the new state transition request.
 				// Callback from terminal state could also trigger a GC that would destroy the state machine
-				FDestinationGameFeaturePluginState::FOnDestinationStateReached LocalOnDestinationStateReached(MoveTemp(DestState->OnDestinationStateReached));
+				UE::GameFeatures::FBroadcastingOnDestinationStateReached LocalOnDestinationStateReached(MoveTemp(DestState->OnDestinationStateReached));
 				DestState->OnDestinationStateReached.Clear();
 
-				LocalOnDestinationStateReached.Broadcast(this, Result);
+				LocalOnDestinationStateReached.CallbackDelegate.Broadcast(this, Result);
 			}
 		}
 	};
@@ -4257,9 +4318,9 @@ void UGameFeaturePluginStateMachine::UpdateStateMachine()
 				DoCallbacks(UE::GameFeatures::CanceledResult, 0, EGameFeaturePluginState::MAX);
 
 				// Must be called after transtition callbacks, UGameFeaturesSubsystem::ChangeGameFeatureTargetStateComplete may remove the this machine from the subsystem
-				FGameFeaturePluginStateMachineProperties::FOnTransitionCanceled LocalOnTransitionCanceled(MoveTemp(StateProperties.OnTransitionCanceled));
+				UE::GameFeatures::FBroadcastingOnTransitionCanceled LocalOnTransitionCanceled(MoveTemp(StateProperties.OnTransitionCanceled));
 				StateProperties.OnTransitionCanceled.Clear();
-				LocalOnTransitionCanceled.Broadcast(this);
+				LocalOnTransitionCanceled.CallbackDelegate.Broadcast(this);
 			}
 			else if (const bool bError = !StateStatus.TransitionResult.HasValue(); bError)
 			{
