Added install cosmetic command.

diff --git a/Engine/Plugins/Runtime/GameFeatures/Source/GameFeatures/Private/GameFeaturesSubsystem.cpp b/Engine/Plugins/Runtime/GameFeatures/Source/GameFeatures/Private/GameFeaturesSubsystem.cpp
--- a/Engine/Plugins/Runtime/GameFeatures/Source/GameFeatures/Private/GameFeaturesSubsystem.cpp
+++ b/Engine/Plugins/Runtime/GameFeatures/Source/GameFeatures/Private/GameFeaturesSubsystem.cpp
@@ -2990,8 +2990,8 @@ struct FGameFeatureAssetStreamContext : public FGameFeaturePluginAssetStreamHand
 			}
 
 			// Check installed and not moving to a state less than installed
-			if (Machine->GetCurrentState() >= EGameFeaturePluginState::Installed &&
-				Machine->GetDestination() > EGameFeaturePluginState::StatusKnown)
+			if (Machine->GetCurrentState() < EGameFeaturePluginState::Installed ||
+				Machine->GetDestination() < EGameFeaturePluginState::Installed)
 			{
 				UE_LOGFMT(LogGameFeatures, Error, "Invalid GFP state for {MountPoint}", ("MountPoint", MountPoint));
 				UE::GameFeatures::FResult ErrorResult = MakeError(FString::Printf(TEXT("%.*s%s"),
diff --git a/Engine/Plugins/Runtime/GameFeatures/Source/GameFeatures/Public/GameFeaturesSubsystem.h b/Engine/Plugins/Runtime/GameFeatures/Source/GameFeatures/Public/GameFeaturesSubsystem.h
--- a/Engine/Plugins/Runtime/GameFeatures/Source/GameFeatures/Public/GameFeaturesSubsystem.h
+++ b/Engine/Plugins/Runtime/GameFeatures/Source/GameFeatures/Public/GameFeaturesSubsystem.h
@@ -286,7 +286,7 @@ struct FGameFeaturePluginPredownloadHandle : public FGameFeaturePluginAsyncHandl
 };
 
 /** Referencer to hold individual asset downloads in local cache */
-struct FGameFeaturePluginStreamedAssetsRef
+struct GAMEFEATURES_API FGameFeaturePluginStreamedAssetsRef
 {
 	explicit FGameFeaturePluginStreamedAssetsRef(FStringView DebugName);
 
