[AutoRTFM] Use a transactionally-safe RWLock when looking up GFP details.

This allows transactional code to inspect the set of currently active GameFeaturePlugins.

diff --git a/Engine/Plugins/Runtime/GameFeatures/Source/GameFeatures/Private/GameFeaturesSubsystem.cpp b/Engine/Plugins/Runtime/GameFeatures/Source/GameFeatures/Private/GameFeaturesSubsystem.cpp
--- a/Engine/Plugins/Runtime/GameFeatures/Source/GameFeatures/Private/GameFeaturesSubsystem.cpp
+++ b/Engine/Plugins/Runtime/GameFeatures/Source/GameFeatures/Private/GameFeaturesSubsystem.cpp
@@ -18,10 +18,11 @@
 #include "Interfaces/IPluginManager.h"
 #include "Logging/StructuredLog.h"
 #include "Misc/App.h"
+#include "Misc/FileHelper.h"
 #include "Misc/PackageName.h"
 #include "Misc/Paths.h"
 #include "Misc/PathViews.h"
-#include "Misc/FileHelper.h"
+#include "Misc/ScopeRWLock.h"
 #include "Serialization/JsonSerializer.h"
 #include "Stats/StatsMisc.h"
 #include "String/ParseTokens.h"
@@ -2390,7 +2391,7 @@ bool UGameFeaturesSubsystem::GetGameFeaturePluginDetailsInternal(const FString&
 
 	if (UE::GameFeatures::GCachePluginDetails)
 	{
-		FReadScopeLock ReadLock(CachedGameFeaturePluginDetailsLock);
+		UE::TReadScopeLock ReadLock(CachedGameFeaturePluginDetailsLock);
 
 		if (FCachedGameFeaturePluginDetails* ExistingDetails = CachedPluginDetailsByFilename.Find(PluginDescriptorFilename))
 		{
@@ -2502,7 +2503,7 @@ bool UGameFeaturesSubsystem::GetGameFeaturePluginDetailsInternal(const FString&
 
 	if (UE::GameFeatures::GCachePluginDetails)
 	{
-		FWriteScopeLock	WriteLock(CachedGameFeaturePluginDetailsLock);
+		UE::TWriteScopeLock WriteLock(CachedGameFeaturePluginDetailsLock);
 		CachedPluginDetailsByFilename.Add(PluginDescriptorFilename, FCachedGameFeaturePluginDetails(OutPluginDetails));
 	}
 
@@ -2511,7 +2512,7 @@ bool UGameFeaturesSubsystem::GetGameFeaturePluginDetailsInternal(const FString&
 
 void UGameFeaturesSubsystem::PruneCachedGameFeaturePluginDetails(const FString& PluginURL, const FString& PluginDescriptorFilename) const
 {
-	FWriteScopeLock WriteLock(CachedGameFeaturePluginDetailsLock);
+	UE::TWriteScopeLock WriteLock(CachedGameFeaturePluginDetailsLock);
 	CachedPluginDetailsByFilename.Remove(PluginDescriptorFilename);
 }
 
diff --git a/Engine/Plugins/Runtime/GameFeatures/Source/GameFeatures/Public/GameFeaturesSubsystem.h b/Engine/Plugins/Runtime/GameFeatures/Source/GameFeatures/Public/GameFeaturesSubsystem.h
--- a/Engine/Plugins/Runtime/GameFeatures/Source/GameFeatures/Public/GameFeaturesSubsystem.h
+++ b/Engine/Plugins/Runtime/GameFeatures/Source/GameFeatures/Public/GameFeaturesSubsystem.h
@@ -6,6 +6,7 @@
 #include "Engine/Engine.h"
 #include "GameFeatureTypesFwd.h"
 #include "IO/IoStoreOnDemand.h"
+#include "Misc/TransactionallySafeRWLock.h"
 
 #include "GameFeaturesSubsystem.generated.h"
 
@@ -935,7 +936,7 @@ private:
 		FCachedGameFeaturePluginDetails(const FGameFeaturePluginDetails& InDetails) : Details(InDetails) {}
 	};
 	mutable TMap<FString, FCachedGameFeaturePluginDetails> CachedPluginDetailsByFilename;
-	mutable FRWLock CachedGameFeaturePluginDetailsLock;
+	mutable FTransactionallySafeRWLock CachedGameFeaturePluginDetailsLock;
 
 	UPROPERTY()
 	TArray<TObjectPtr<UObject>> Observers;
