Tweak EGameFeatureAssetStreamPriority enum to use unsigned integer and expose to UE for use in initial data setup.

diff --git a/Engine/Plugins/Runtime/GameFeatures/Source/GameFeatures/Public/GameFeaturesSubsystem.h b/Engine/Plugins/Runtime/GameFeatures/Source/GameFeatures/Public/GameFeaturesSubsystem.h
--- a/Engine/Plugins/Runtime/GameFeatures/Source/GameFeatures/Public/GameFeaturesSubsystem.h
+++ b/Engine/Plugins/Runtime/GameFeatures/Source/GameFeatures/Public/GameFeaturesSubsystem.h
@@ -308,13 +308,14 @@ struct FGameFeaturePluginAssetStreamHandle : public FGameFeaturePluginAsyncHandl
 	virtual FGameFeaturePluginStreamedAssetsRef GetAssetsRef() const = 0;
 };
 
-enum class EGameFeatureAssetStreamPriority : int8
+UENUM(BlueprintType)
+enum class EGameFeatureAssetStreamPriority : uint8
 {
-	Lowest = -128,
-	BelowNormal = -1,
-	Normal = 0,
-	AboveNormal = 1,
-	Highest = 127
+	Lowest = 0,
+	BelowNormal = 127,
+	Normal = 128,
+	AboveNormal = 129,
+	Highest = 255
 };
 
 /** Struct used to transform a GameFeaturePlugin URL into something that can uniquely identify the GameFeaturePlugin
