FGameFeaturePluginState_Downloading
- Don't fire ensure about async work unless actually starting async work.
- Only call OnGameFeatureDownloading when starting a download

diff --git a/Engine/Plugins/Runtime/GameFeatures/Source/GameFeatures/Private/GameFeaturePluginStateMachine.cpp b/Engine/Plugins/Runtime/GameFeatures/Source/GameFeatures/Private/GameFeaturePluginStateMachine.cpp
--- a/Engine/Plugins/Runtime/GameFeatures/Source/GameFeatures/Private/GameFeaturePluginStateMachine.cpp
+++ b/Engine/Plugins/Runtime/GameFeatures/Source/GameFeatures/Private/GameFeaturePluginStateMachine.cpp
@@ -1520,6 +1520,11 @@ struct FGameFeaturePluginState_Downloading : public FGameFeaturePluginState
 	};
 	TSharedPtr<FIoStoreOnDemandContext> IoStoreOnDemandContext;
 
+	void EnsureAllowAsyncLoading() const
+	{
+		ensureMsgf(AllowAsyncLoading(), TEXT("FGameFeaturePluginState::AllowAsyncLoading is false while attempting to download GFP data."));
+	}
+
 	void Cleanup()
 	{
 		if (ProgressUpdateHandle.IsValid())
@@ -1569,15 +1574,6 @@ struct FGameFeaturePluginState_Downloading : public FGameFeaturePluginState
 		}
 
 		const FInstallBundlePluginProtocolMetaData& MetaData = StateProperties.ProtocolMetadata.GetSubtype<FInstallBundlePluginProtocolMetaData>();
-		const FInstallBundlePluginProtocolOptions& Options = StateProperties.ProtocolOptions.GetSubtype<FInstallBundlePluginProtocolOptions>();
-
-		// if not up to date, check to see if we allow downloading
-		if (Options.bDoNotDownload && !BundleContentState.GetAllBundlesHaveState(EInstallBundleInstallState::UpToDate))
-		{
-			Result = GetErrorResult(TEXT("GFPStateMachine.DownloadNotAllowed"));
-			UpdateStateMachineImmediate();
-			return;
-		}
 
 		UE::GameFeatures::GFPSharedInstallTracker.AddBundleRefs(StateProperties.PluginName, UE::GameFeatures::EGFPInstallLevel::Download, MetaData.InstallBundles);
 
@@ -1625,6 +1621,7 @@ struct FGameFeaturePluginState_Downloading : public FGameFeaturePluginState
 
 			//If this setting is flipped then we should immediately request to pause downloads.
 			//We still generate the downloads so that we have an accurate PendingBundleDownloads list
+			const FInstallBundlePluginProtocolOptions& Options = StateProperties.ProtocolOptions.GetSubtype<FInstallBundlePluginProtocolOptions>();
 			if (Options.bUserPauseDownload)
 			{
 				ChangePauseState(true);
@@ -1798,27 +1795,59 @@ struct FGameFeaturePluginState_Downloading : public FGameFeaturePluginState
 			return;
 		}
 
-		UGameFeaturesSubsystem::Get().OnGameFeatureDownloading(StateProperties.PluginName, StateProperties.PluginIdentifier);
-
-		ensureMsgf(AllowAsyncLoading(), TEXT("FGameFeaturePluginState::AllowAsyncLoading is false while attempting to download GFP data."));
-
 		TSharedPtr<IInstallBundleManager> BundleManager = IInstallBundleManager::GetPlatformInstallBundleManager();
 		const TConstArrayView<FName> InstallBundles = GetInstallBundles();
 
+		const FInstallBundlePluginProtocolOptions& Options = StateProperties.ProtocolOptions.GetSubtype<FInstallBundlePluginProtocolOptions>();
+
 		const bool bAddDependencies = false; // We already got all dependencies in the CheckingStatus state
 		if (InstallBundles.Num() > 1)
 		{
+			// Handle bDoNotDownload flag before doing any async ops
+			if (Options.bDoNotDownload)
+			{
+				TValueOrError<FInstallBundleCombinedInstallState, EInstallBundleResult> MaybeInstallState = 
+					BundleManager->GetInstallStateSynchronous(InstallBundles, bAddDependencies);
+				check(MaybeInstallState.HasValue());
+				FInstallBundleCombinedInstallState& InstallState = MaybeInstallState.GetValue();
+
+				// if not up to date, check to see if we allow downloading
+				if (!InstallState.GetAllBundlesHaveState(EInstallBundleInstallState::UpToDate))
+				{
+					Result = GetErrorResult(TEXT("GFPStateMachine.DownloadNotAllowed"));
+					UpdateStateMachineImmediate();
+					return;
+				}
+			}
+
+			UGameFeaturesSubsystem::Get().OnGameFeatureDownloading(StateProperties.PluginName, StateProperties.PluginIdentifier);
+			EnsureAllowAsyncLoading();
+
 			GotContentStateHandle = BundleManager->GetContentState(InstallBundles, EInstallBundleGetContentStateFlags::None, bAddDependencies,
 				FInstallBundleGetContentStateDelegate::CreateRaw(this, &FGameFeaturePluginState_Downloading::OnGotContentState));
 		}
 		else
 		{
+			TValueOrError<FInstallBundleCombinedInstallState, EInstallBundleResult> MaybeInstallState = 
+				BundleManager->GetInstallStateSynchronous(InstallBundles, bAddDependencies);
+			check(MaybeInstallState.HasValue());
+			FInstallBundleCombinedInstallState& InstallState = MaybeInstallState.GetValue();
+
+			// Handle bDoNotDownload flag before doing any async ops
+			// if not up to date, check to see if we allow downloading
+			if (Options.bDoNotDownload && !InstallState.GetAllBundlesHaveState(EInstallBundleInstallState::UpToDate))
+			{
+				Result = GetErrorResult(TEXT("GFPStateMachine.DownloadNotAllowed"));
+				UpdateStateMachineImmediate();
+				return;
+			}
+
+			UGameFeaturesSubsystem::Get().OnGameFeatureDownloading(StateProperties.PluginName, StateProperties.PluginIdentifier);
+			EnsureAllowAsyncLoading();
+
 			// We usualy only have a use case of one bundle per GFP and we only care about relative weighting here, so we don't need 
 			// any of the other content state metadata. We can just assume the weight is 1.0 and not have to wait for the full 
 			// async call to get the rest of the metadata.
-			TValueOrError<FInstallBundleCombinedInstallState, EInstallBundleResult> MaybeInstallState = BundleManager->GetInstallStateSynchronous(InstallBundles, bAddDependencies);
-			check(MaybeInstallState.HasValue());
-			FInstallBundleCombinedInstallState& InstallState = MaybeInstallState.GetValue();
 			FInstallBundleCombinedContentState HackContentState;
 			HackContentState.IndividualBundleStates.Reserve(InstallState.IndividualBundleStates.Num());
 			for (const TPair<FName, EInstallBundleInstallState>& Pair : InstallState.IndividualBundleStates)
