[Optional Bundles Installation] Added the necessary events/functions to react to optional bundles being downloaded/installed.

#rnx

diff --git a/Engine/Plugins/Runtime/GameFeatures/Source/GameFeatures/Private/GameFeatureOptionalContentInstaller.cpp b/Engine/Plugins/Runtime/GameFeatures/Source/GameFeatures/Private/GameFeatureOptionalContentInstaller.cpp
--- a/Engine/Plugins/Runtime/GameFeatures/Source/GameFeatures/Private/GameFeatureOptionalContentInstaller.cpp
+++ b/Engine/Plugins/Runtime/GameFeatures/Source/GameFeatures/Private/GameFeatureOptionalContentInstaller.cpp
@@ -22,6 +22,8 @@ namespace GameFeatureOptionalContentInstaller
 }
 
 TMulticastDelegate<void(const FString& PluginName, const UE::GameFeatures::FResult&)> UGameFeatureOptionalContentInstaller::OnOptionalContentInstalled;
+TMulticastDelegate<void()> UGameFeatureOptionalContentInstaller::OnOptionalContentInstallStarted;
+TMulticastDelegate<void()> UGameFeatureOptionalContentInstaller::OnOptionalContentInstallFinished;
 
 void UGameFeatureOptionalContentInstaller::BeginDestroy()
 {
@@ -89,6 +91,11 @@ void UGameFeatureOptionalContentInstaller::EnableCellularDownloading(bool bEnabl
 	}
 }
 
+bool UGameFeatureOptionalContentInstaller::HasOngoingInstalls() const
+{
+	return ActiveGFPInstalls.Num() > 0;
+}
+
 bool UGameFeatureOptionalContentInstaller::UpdateContent(const FString& PluginName, bool bIsPredownload)
 {
 	TArray<FName> Bundles = GetOptionalBundlePredicate(PluginName);
@@ -145,6 +152,11 @@ bool UGameFeatureOptionalContentInstaller::UpdateContent(const FString& PluginNa
 	FInstallBundleRequestInfo& Request =  MaybeRequest.GetValue();
 	if (!Request.BundlesEnqueued.IsEmpty())
 	{
+		if (ActiveGFPInstalls.Num() == 0)
+    	{
+    		OnOptionalContentInstallStarted.Broadcast();
+    	}
+	
 		FGFPInstall& Pending = ActiveGFPInstalls.FindOrAdd(PluginName);
 
 		if (!Pending.CallbackHandle.IsValid())
@@ -240,6 +252,11 @@ void UGameFeatureOptionalContentInstaller::OnContentInstalled(FInstallBundleRequ
 		// book keeping
 		IInstallBundleManager::InstallBundleCompleteDelegate.Remove(GFPInstall.CallbackHandle);
 		ActiveGFPInstalls.Remove(PluginName);
+
+		if (ActiveGFPInstalls.Num() == 0)
+		{
+			OnOptionalContentInstallFinished.Broadcast();
+		}
 	}
 }
 
diff --git a/Engine/Plugins/Runtime/GameFeatures/Source/GameFeatures/Public/GameFeatureOptionalContentInstaller.h b/Engine/Plugins/Runtime/GameFeatures/Source/GameFeatures/Public/GameFeatureOptionalContentInstaller.h
--- a/Engine/Plugins/Runtime/GameFeatures/Source/GameFeatures/Public/GameFeatureOptionalContentInstaller.h
+++ b/Engine/Plugins/Runtime/GameFeatures/Source/GameFeatures/Public/GameFeatureOptionalContentInstaller.h
@@ -26,6 +26,8 @@ class UGameFeatureOptionalContentInstaller : public UObject, public IGameFeature
 
 public:
 	static GAMEFEATURES_API TMulticastDelegate<void(const FString& PluginName, const UE::GameFeatures::FResult&)> OnOptionalContentInstalled;
+	static GAMEFEATURES_API TMulticastDelegate<void()> OnOptionalContentInstallStarted;
+	static GAMEFEATURES_API TMulticastDelegate<void()> OnOptionalContentInstallFinished;
 
 public:
 	virtual void BeginDestroy() override;
@@ -37,6 +39,8 @@ public:
 	void GAMEFEATURES_API UninstallContent();
 
 	void GAMEFEATURES_API EnableCellularDownloading(bool bEnable);
+	
+	bool GAMEFEATURES_API HasOngoingInstalls() const;
 
 private:
 	bool UpdateContent(const FString& PluginName, bool bIsPredownload);
