Cleanup IoStoreOnDemandModule usage by GFP subsystem.
Fix cancel behavior for FGameFeatureAssetStreamContext.

diff --git a/Engine/Plugins/Runtime/GameFeatures/Source/GameFeatures/Private/GameFeaturePluginStateMachine.cpp b/Engine/Plugins/Runtime/GameFeatures/Source/GameFeatures/Private/GameFeaturePluginStateMachine.cpp
--- a/Engine/Plugins/Runtime/GameFeatures/Source/GameFeatures/Private/GameFeaturePluginStateMachine.cpp
+++ b/Engine/Plugins/Runtime/GameFeatures/Source/GameFeatures/Private/GameFeaturePluginStateMachine.cpp
@@ -1501,8 +1501,6 @@ struct FGameFeaturePluginState_Downloading : public FGameFeaturePluginState
 	FTSTicker::FDelegateHandle ProgressUpdateHandle;
 	FDelegateHandle GotContentStateHandle;
 
-	UE::IoStore::FIoStoreOnDemandModule* IoStoreOnDemandModule = nullptr;
-
 	// Required for callback lifetime safety
 	struct FIoStoreOnDemandContext
 	{
@@ -1550,7 +1548,6 @@ struct FGameFeaturePluginState_Downloading : public FGameFeaturePluginState
 		bPluginDownloaded = false;
 		PendingBundleDownloads.Empty();
 		ProgressTracker = nullptr;
-		IoStoreOnDemandModule = nullptr;
 
 		if (IoStoreOnDemandContext)
 		{
@@ -1672,8 +1669,8 @@ struct FGameFeaturePluginState_Downloading : public FGameFeaturePluginState
 			return;
 		}
 
-		IoStoreOnDemandModule = FModuleManager::Get().GetModulePtr<UE::IoStore::FIoStoreOnDemandModule>(TEXT("IoStoreOnDemand"));
-		if (!IoStoreOnDemandModule && !Result.HasError())
+		UE::IoStore::IOnDemandIoStore* IoStore = UE::IoStore::TryGetOnDemandIoStore();
+		if (!IoStore && !Result.HasError())
 		{
 			Result = GetErrorResult(TEXT("IoStoreOnDemand.ModuleNotFound"));
 			UpdateStateMachineImmediate();
@@ -1683,7 +1680,6 @@ struct FGameFeaturePluginState_Downloading : public FGameFeaturePluginState
 		IoStoreOnDemandContext = MakeShared<FIoStoreOnDemandContext>();
 		IoStoreOnDemandContext->PendingInstalls = MetaData.InstallBundlesWithAssetDependencies.Num();
 
-		UE::IoStore::IOnDemandIoStore& IoStore = UE::IoStore::GetOnDemandIoStore();
 		for (const FName InstallBundle : MetaData.InstallBundlesWithAssetDependencies)
 		{
 			UE::IoStore::FOnDemandInstallArgs InstallArgs;
@@ -1698,7 +1694,7 @@ struct FGameFeaturePluginState_Downloading : public FGameFeaturePluginState
 
 			// This should be pretty small, so not going to worry about progress here.
 			IoStoreOnDemandContext->InstallRequests.Add(
-				IoStore.Install(MoveTemp(InstallArgs),
+				IoStore->Install(MoveTemp(InstallArgs),
 				[this, LambdaOnDemandContext = IoStoreOnDemandContext](const UE::IoStore::FOnDemandInstallResult& OnDemandInstallResult)
 				{
 					if (!LambdaOnDemandContext->bStateValid)
@@ -2838,7 +2834,6 @@ struct FGameFeaturePluginState_AssetDependencyStreaming : public FGameFeaturePlu
 		bool bStateValid = true;
 	};
 
-	UE::IoStore::FIoStoreOnDemandModule* IoStoreOnDemandModule = nullptr;
 	TSharedPtr<FIoStoreOnDemandContext> IoStoreOnDemandContext;
 
 	UE::GameFeatures::FResult Result{ MakeValue() };
@@ -2849,8 +2844,6 @@ struct FGameFeaturePluginState_AssetDependencyStreaming : public FGameFeaturePlu
 		Result = MakeValue();
 		bComplete = false;
 
-		IoStoreOnDemandModule = nullptr;
-
 		if (IoStoreOnDemandContext)
 		{
 			for (UE::IoStore::FOnDemandInstallRequest& Request : IoStoreOnDemandContext->InstallRequests)
@@ -2880,8 +2873,8 @@ struct FGameFeaturePluginState_AssetDependencyStreaming : public FGameFeaturePlu
 			return;
 		}
 
-		IoStoreOnDemandModule = FModuleManager::Get().GetModulePtr<UE::IoStore::FIoStoreOnDemandModule>(TEXT("IoStoreOnDemand"));
-		if (!IoStoreOnDemandModule)
+		UE::IoStore::IOnDemandIoStore* IoStore = UE::IoStore::TryGetOnDemandIoStore();
+		if (!IoStore)
 		{
 			Result = GetErrorResult(TEXT("IoStoreOnDemand.ModuleNotFound"));
 			return;
@@ -2905,7 +2898,6 @@ struct FGameFeaturePluginState_AssetDependencyStreaming : public FGameFeaturePlu
 		IoStoreOnDemandContext->Progress.Reserve(MetaData.InstallBundlesWithAssetDependencies.Num());
 
 		const TArray<EStreamingAssetInstallMode>& InstallModes = MaybeInstallModes.GetValue();
-		UE::IoStore::IOnDemandIoStore& IoStore = UE::IoStore::GetOnDemandIoStore();
 		for (int i = 0; const FName InstallBundle : MetaData.InstallBundlesWithAssetDependencies)
 		{
 			IoStoreOnDemandContext->Progress.Emplace(FIoStoreOnDemandProgress{InstallBundle});
@@ -2924,7 +2916,7 @@ struct FGameFeaturePluginState_AssetDependencyStreaming : public FGameFeaturePlu
 			InstallArgs.ContentHandle = UE::GameFeatures::GFPSharedInstallTracker.AddOnDemandContentHandle(
 				InstallBundle, UE::GameFeatures::EGFPInstallLevel::AssetStream);
 
-			IoStoreOnDemandContext->InstallRequests.Add(IoStore.Install(MoveTemp(InstallArgs),
+			IoStoreOnDemandContext->InstallRequests.Add(IoStore->Install(MoveTemp(InstallArgs),
 				// On Complete
 				[this, LambdaOnDemandContext = IoStoreOnDemandContext](const UE::IoStore::FOnDemandInstallResult& OnDemandInstallResult)
 				{
diff --git a/Engine/Plugins/Runtime/GameFeatures/Source/GameFeatures/Private/GameFeaturesSubsystem.cpp b/Engine/Plugins/Runtime/GameFeatures/Source/GameFeatures/Private/GameFeaturesSubsystem.cpp
--- a/Engine/Plugins/Runtime/GameFeatures/Source/GameFeatures/Private/GameFeaturesSubsystem.cpp
+++ b/Engine/Plugins/Runtime/GameFeatures/Source/GameFeatures/Private/GameFeaturesSubsystem.cpp
@@ -2882,9 +2882,10 @@ struct FGameFeatureAssetStreamContext : public FGameFeaturePluginAssetStreamHand
 	TUniqueFunction<void(float)> OnProgress;
 	UE::GameFeatures::FResult Result = MakeValue();
 	TArray<UE::IoStore::FOnDemandInstallProgress> ProgressList;
+	TArray<UE::IoStore::FOnDemandInstallRequest> InstallRequests;
 	int32 PendingInstalls = 0;
 	bool bIsComplete = false;
-	TArray<UE::IoStore::FOnDemandInstallRequest> InstallRequests;
+	bool bIsCanceled = false;
 
 	FGameFeatureAssetStreamContext(FGameFeaturePluginStreamedAssetsRef&& InAssetRef)
 		: AssetRef(MoveTemp(InAssetRef))
@@ -2909,6 +2910,7 @@ struct FGameFeatureAssetStreamContext : public FGameFeaturePluginAssetStreamHand
 
 	virtual void Cancel() override
 	{
+		bIsCanceled = true;
 		for (UE::IoStore::FOnDemandInstallRequest& Request : InstallRequests)
 		{
 			Request.Cancel();
@@ -2922,7 +2924,7 @@ struct FGameFeatureAssetStreamContext : public FGameFeaturePluginAssetStreamHand
 
 	bool IsCanceled()
 	{
-		return Algo::AnyOf(InstallRequests, [](const UE::IoStore::FOnDemandInstallRequest& R) { return R.IsCancelled(); });
+		return bIsCanceled;
 	}
 
 	void SetComplete()
@@ -3115,8 +3117,8 @@ struct FGameFeatureAssetStreamContext : public FGameFeaturePluginAssetStreamHand
 
 	void Start(TArrayView<UE::IoStore::FOnDemandInstallArgs> InstallArgList, EGameFeatureAssetStreamPriority Priority)
 	{
-		UE::IoStore::FIoStoreOnDemandModule* IoStoreOnDemandModule = FModuleManager::Get().GetModulePtr<UE::IoStore::FIoStoreOnDemandModule>(TEXT("IoStoreOnDemand"));
-		if (!IoStoreOnDemandModule)
+		UE::IoStore::IOnDemandIoStore* IoStore = UE::IoStore::TryGetOnDemandIoStore();
+		if (!IoStore)
 		{
 			UE::GameFeatures::FResult ErrorResult = MakeError(FString::Printf(TEXT("%.*s%s"),
 				ErrorNamespace.Len(), ErrorNamespace.GetData(), TEXT("IoStoreOnDemand.ModuleNotFound")));
@@ -3136,8 +3138,8 @@ struct FGameFeatureAssetStreamContext : public FGameFeaturePluginAssetStreamHand
 
 		PendingInstalls = InstallArgList.Num();
 		ProgressList.Empty(PendingInstalls);
+		InstallRequests.Empty(PendingInstalls);
 
-		UE::IoStore::IOnDemandIoStore& IoStore = UE::IoStore::GetOnDemandIoStore();
 		for (UE::IoStore::FOnDemandInstallArgs& InstallArgs : InstallArgList)
 		{
 			const int32 RequestIndex = ProgressList.Emplace();
@@ -3148,7 +3150,7 @@ struct FGameFeatureAssetStreamContext : public FGameFeaturePluginAssetStreamHand
 
 			// TODO: convert Priority to IOStore priority here
 
-			InstallRequests.Add(IoStore.Install(MoveTemp(InstallArgs),
+			InstallRequests.Add(IoStore->Install(MoveTemp(InstallArgs),
 				[Context = SharedThis(this), RequestIndex](const UE::IoStore::FOnDemandInstallResult& OnDemandInstallResult)
 				{
 					Context->OnInstallComplete(OnDemandInstallResult, RequestIndex);
