[Optional Bundles Installation] Added the necessary events/functions to react to optional bundles being downloaded/installed.

#rnx

diff --git a/Engine/Plugins/Runtime/GameFeatures/Source/GameFeatures/Private/GameFeatureOptionalContentInstaller.cpp b/Engine/Plugins/Runtime/GameFeatures/Source/GameFeatures/Private/GameFeatureOptionalContentInstaller.cpp
--- a/Engine/Plugins/Runtime/GameFeatures/Source/GameFeatures/Private/GameFeatureOptionalContentInstaller.cpp
+++ b/Engine/Plugins/Runtime/GameFeatures/Source/GameFeatures/Private/GameFeatureOptionalContentInstaller.cpp
@@ -23,7 +23,7 @@ namespace GameFeatureOptionalContentInstaller
 
 TMulticastDelegate<void(const FString& PluginName, const UE::GameFeatures::FResult&)> UGameFeatureOptionalContentInstaller::OnOptionalContentInstalled;
 TMulticastDelegate<void()> UGameFeatureOptionalContentInstaller::OnOptionalContentInstallStarted;
-TMulticastDelegate<void()> UGameFeatureOptionalContentInstaller::OnOptionalContentInstallFinished;
+TMulticastDelegate<void(const bool bInstallSuccessful)> UGameFeatureOptionalContentInstaller::OnOptionalContentInstallFinished;
 
 void UGameFeatureOptionalContentInstaller::BeginDestroy()
 {
@@ -255,7 +255,8 @@ void UGameFeatureOptionalContentInstaller::OnContentInstalled(FInstallBundleRequ
 
 		if (ActiveGFPInstalls.Num() == 0)
 		{
-			OnOptionalContentInstallFinished.Broadcast();
+			const bool bInstallSuccessful = RelevantGFPs.Num() > 0;
+			OnOptionalContentInstallFinished.Broadcast(bInstallSuccessful);
 		}
 	}
 }
diff --git a/Engine/Plugins/Runtime/GameFeatures/Source/GameFeatures/Public/GameFeatureOptionalContentInstaller.h b/Engine/Plugins/Runtime/GameFeatures/Source/GameFeatures/Public/GameFeatureOptionalContentInstaller.h
--- a/Engine/Plugins/Runtime/GameFeatures/Source/GameFeatures/Public/GameFeatureOptionalContentInstaller.h
+++ b/Engine/Plugins/Runtime/GameFeatures/Source/GameFeatures/Public/GameFeatureOptionalContentInstaller.h
@@ -27,7 +27,7 @@ class UGameFeatureOptionalContentInstaller : public UObject, public IGameFeature
 public:
 	static GAMEFEATURES_API TMulticastDelegate<void(const FString& PluginName, const UE::GameFeatures::FResult&)> OnOptionalContentInstalled;
 	static GAMEFEATURES_API TMulticastDelegate<void()> OnOptionalContentInstallStarted;
-	static GAMEFEATURES_API TMulticastDelegate<void()> OnOptionalContentInstallFinished;
+	static GAMEFEATURES_API TMulticastDelegate<void(const bool bInstallSuccessful)> OnOptionalContentInstallFinished;
 
 public:
 	virtual void BeginDestroy() override;
