UGameFeaturesSubsystem::LoadBuiltInGameFeaturePlugin
- Fix protocol options not being updated on plugin if already exists.

diff --git a/Engine/Plugins/Runtime/GameFeatures/Source/GameFeatures/Private/GameFeaturesSubsystem.cpp b/Engine/Plugins/Runtime/GameFeatures/Source/GameFeatures/Private/GameFeaturesSubsystem.cpp
--- a/Engine/Plugins/Runtime/GameFeatures/Source/GameFeatures/Private/GameFeaturesSubsystem.cpp
+++ b/Engine/Plugins/Runtime/GameFeatures/Source/GameFeatures/Private/GameFeaturesSubsystem.cpp
@@ -1983,6 +1983,10 @@ void UGameFeaturesSubsystem::LoadBuiltInGameFeaturePlugin(const TSharedRef<IPlug
 #endif // WITH_EDITOR
 
 	UAssetManager::Get().PushBulkScanning();
+	ON_SCOPE_EXIT
+	{
+		UAssetManager::Get().PopBulkScanning();
+	};
 
 	const FString& PluginDescriptorFilename = Plugin->GetDescriptorFileName();
 
@@ -2039,7 +2043,18 @@ void UGameFeaturesSubsystem::LoadBuiltInGameFeaturePlugin(const TSharedRef<IPlug
 			ProtocolOptions.bLogWarningOnForcedDependencyCreation = BehaviorOptions.bLogWarningOnForcedDependencyCreation;
 			ProtocolOptions.bLogErrorOnForcedDependencyCreation = BehaviorOptions.bLogErrorOnForcedDependencyCreation;
 			ProtocolOptions.bBatchProcess = BehaviorOptions.bBatchProcess;
-			UGameFeaturePluginStateMachine* StateMachine = FindOrCreateGameFeaturePluginStateMachine(PluginURL, ProtocolOptions);
+
+			bool bFoundExisting = false;
+			UGameFeaturePluginStateMachine* StateMachine = FindOrCreateGameFeaturePluginStateMachine(PluginURL, ProtocolOptions, &bFoundExisting);
+			if (bFoundExisting && ShouldUpdatePluginProtocolOptions(StateMachine, ProtocolOptions))
+			{
+				const UE::GameFeatures::FResult Result = UpdateGameFeatureProtocolOptions(StateMachine, ProtocolOptions);
+				if (Result.HasError())
+				{
+					CompleteDelegate.ExecuteIfBound(Result);
+					return;
+				}
+			}
 
 			EBuiltInAutoState InitialAutoState = (BehaviorOptions.AutoStateOverride != EBuiltInAutoState::Invalid) ? 
 				BehaviorOptions.AutoStateOverride : PluginDetails.BuiltInAutoState;
@@ -2062,8 +2077,6 @@ void UGameFeaturesSubsystem::LoadBuiltInGameFeaturePlugin(const TSharedRef<IPlug
 				}));
 		}
 	}
-
-	UAssetManager::Get().PopBulkScanning();
 }
 
 #if UE_BUILD_SHIPPING
