Sort GFP dependency state machines using `Algo::Sort` instead of `TArray::Sort`, saving 225K FString allocations in an internal project.
`TArray::Sort` dereferences pointers when sorting, which results in `UObjectBaseUtility::operator<` being used, which does a lexographic comparison of the object names via FString.

diff --git a/Engine/Plugins/Runtime/GameFeatures/Source/GameFeatures/Private/GameFeaturesSubsystem.cpp b/Engine/Plugins/Runtime/GameFeatures/Source/GameFeatures/Private/GameFeaturesSubsystem.cpp
--- a/Engine/Plugins/Runtime/GameFeatures/Source/GameFeatures/Private/GameFeaturesSubsystem.cpp
+++ b/Engine/Plugins/Runtime/GameFeatures/Source/GameFeatures/Private/GameFeaturesSubsystem.cpp
@@ -3531,7 +3531,7 @@ bool UGameFeaturesSubsystem::FindOrCreatePluginDependencyStateMachines(const FSt
 			OutDependencyMachines.Add(ResolvedDependency);
 		}
 
-		OutDependencyMachines.Sort();
+		Algo::Sort(OutDependencyMachines);
 		OutDependencyMachines.SetNum(Algo::Unique(OutDependencyMachines));
 		return true;
 	}
@@ -3555,7 +3555,7 @@ bool UGameFeaturesSubsystem::FindPluginDependencyStateMachinesToActivate(const F
 			return false;
 		}
 	});
-	OutDependencyMachines.Sort();
+	Algo::Sort(OutDependencyMachines);
 	OutDependencyMachines.SetNum(Algo::Unique(OutDependencyMachines));
 
 	return bResult;
@@ -3576,7 +3576,7 @@ bool UGameFeaturesSubsystem::FindPluginDependencyStateMachinesToDeactivate(const
 		}
 		return true;
 	});
-	OutDependencyMachines.Sort();
+	Algo::Sort(OutDependencyMachines);
 	OutDependencyMachines.SetNum(Algo::Unique(OutDependencyMachines));
 
 	return bResult;
