[Optional Content Download] Added a function to get the % of progress.

#rnx

diff --git a/Engine/Plugins/Runtime/GameFeatures/Source/GameFeatures/Private/GameFeatureOptionalContentInstaller.cpp b/Engine/Plugins/Runtime/GameFeatures/Source/GameFeatures/Private/GameFeatureOptionalContentInstaller.cpp
--- a/Engine/Plugins/Runtime/GameFeatures/Source/GameFeatures/Private/GameFeatureOptionalContentInstaller.cpp
+++ b/Engine/Plugins/Runtime/GameFeatures/Source/GameFeatures/Private/GameFeatureOptionalContentInstaller.cpp
@@ -21,6 +21,8 @@ namespace GameFeatureOptionalContentInstaller
 		TEXT("Enable optional content installer"));
 }
 
+const FName UGameFeatureOptionalContentInstaller::GFOContentRequestName = TEXT("GFOContentRequest");
+
 TMulticastDelegate<void(const FString& PluginName, const UE::GameFeatures::FResult&)> UGameFeatureOptionalContentInstaller::OnOptionalContentInstalled;
 TMulticastDelegate<void()> UGameFeatureOptionalContentInstaller::OnOptionalContentInstallStarted;
 TMulticastDelegate<void(const bool bInstallSuccessful)> UGameFeatureOptionalContentInstaller::OnOptionalContentInstallFinished;
@@ -96,6 +98,24 @@ bool UGameFeatureOptionalContentInstaller::HasOngoingInstalls() const
 	return ActiveGFPInstalls.Num() > 0;
 }
 
+float UGameFeatureOptionalContentInstaller::GetAllInstallsProgress()
+{
+	if (TotalProgressTracker.IsSet())
+	{
+		TotalProgressTracker->ForceTick();
+		return TotalProgressTracker->GetCurrentCombinedProgress().ProgressPercent;
+	}
+
+	if (ActiveGFPInstalls.Num() > 0 && RelevantGFPs.Num() > 0)
+	{
+		// Start the tracker for next calls to this function
+		StartTotalProgressTracker();
+	}
+
+	// Return 1 if some optional bundles are installed, 0 if none are installed 
+	return RelevantGFPs.Num() > 0 ? 1.f : 0.f;
+}
+
 bool UGameFeatureOptionalContentInstaller::UpdateContent(const FString& PluginName, bool bIsPredownload)
 {
 	TArray<FName> Bundles = GetOptionalBundlePredicate(PluginName);
@@ -257,6 +277,7 @@ void UGameFeatureOptionalContentInstaller::OnContentInstalled(FInstallBundleRequ
 		{
 			const bool bInstallSuccessful = RelevantGFPs.Num() > 0;
 			OnOptionalContentInstallFinished.Broadcast(bInstallSuccessful);
+			TotalProgressTracker.Reset();
 		}
 	}
 }
@@ -301,6 +322,7 @@ void UGameFeatureOptionalContentInstaller::OnDisabled()
 	}
 
 	RelevantGFPs.Empty();
+	TotalProgressTracker.Reset();
 }
 
 bool UGameFeatureOptionalContentInstaller::IsEnabled() const
@@ -313,6 +335,46 @@ void UGameFeatureOptionalContentInstaller::OnCVarsChanged()
 	Enable(bEnabled); // Check if CVar changed IsEnabled() and if so, call callbacks
 }
 
+void UGameFeatureOptionalContentInstaller::StartTotalProgressTracker()
+{
+	TotalProgressTracker.Reset();
+	BundleManager->CancelAllGetContentStateRequestsForTag(GFOContentRequestName);
+	
+	TArray<FName> AllActiveBundleInstalls;
+	for (const TTuple<FString, FGFPInstall>& ActiveInstall : ActiveGFPInstalls)
+	{
+		for (const FName& BundleEnqueued : ActiveInstall.Value.BundlesEnqueued)
+		{
+			AllActiveBundleInstalls.AddUnique(BundleEnqueued);
+		}
+	}
+
+	if (AllActiveBundleInstalls.Num() > 0 && RelevantGFPs.Num() > 0)
+	{
+		// Start a new progress tracker for the currently active bundle installs. Auto tick is disabled.
+		TotalProgressTracker.Emplace(false);
+		
+		TWeakObjectPtr This_WeakPtr = this;	
+		BundleManager->GetContentState(AllActiveBundleInstalls, EInstallBundleGetContentStateFlags::None, false,
+			FInstallBundleGetContentStateDelegate::CreateLambda([This_WeakPtr](FInstallBundleCombinedContentState BundleContentState)
+			{
+				if (This_WeakPtr.IsValid())
+				{
+					const TStrongObjectPtr<UGameFeatureOptionalContentInstaller> This_StrongPtr = This_WeakPtr.Pin();
+					if (This_StrongPtr->TotalProgressTracker.IsSet())
+					{
+						TArray<FName> RequiredBundlesForTracking;
+						for (const TPair<FName, FInstallBundleContentState>& BundleState : BundleContentState.IndividualBundleStates)
+						{
+							RequiredBundlesForTracking.Add(BundleState.Key);
+						}
+						This_StrongPtr->TotalProgressTracker->SetBundlesToTrackFromContentState(BundleContentState, MoveTemp(RequiredBundlesForTracking));
+					}
+				}
+			}), GFOContentRequestName);
+	}
+}
+
 void UGameFeatureOptionalContentInstaller::OnGameFeaturePredownloading(const FString& PluginName, const FGameFeaturePluginIdentifier& PluginIdentifier)
 {
 	if (!IsEnabled())
diff --git a/Engine/Plugins/Runtime/GameFeatures/Source/GameFeatures/Public/GameFeatureOptionalContentInstaller.h b/Engine/Plugins/Runtime/GameFeatures/Source/GameFeatures/Public/GameFeatureOptionalContentInstaller.h
--- a/Engine/Plugins/Runtime/GameFeatures/Source/GameFeatures/Public/GameFeatureOptionalContentInstaller.h
+++ b/Engine/Plugins/Runtime/GameFeatures/Source/GameFeatures/Public/GameFeatureOptionalContentInstaller.h
@@ -2,6 +2,7 @@
 
 #pragma once
 
+#include "BundlePrereqCombinedStatusHelper.h"
 #include "GameFeatureStateChangeObserver.h"
 #include "HAL/IConsoleManager.h"
 #include "InstallBundleManagerInterface.h"
@@ -42,6 +43,12 @@ public:
 	
 	bool GAMEFEATURES_API HasOngoingInstalls() const;
 
+	/**
+	 * Returns the total progress (between 0 and 1) of all the optional bundles currently being installed.
+	 * Returns 0 while the progress tracker is starting which happens the first time it is called while bundles are being installed.
+	 */
+	float GAMEFEATURES_API GetAllInstallsProgress();
+
 private:
 	bool UpdateContent(const FString& PluginName, bool bIsPredownload);
 
@@ -56,6 +63,8 @@ private:
 
 	void OnCVarsChanged();
 
+	void StartTotalProgressTracker();
+
 	// IGameFeatureStateChangeObserver Interface
 	virtual void OnGameFeaturePredownloading(const FString& PluginName, const FGameFeaturePluginIdentifier& PluginIdentifier) override;
 	virtual void OnGameFeatureDownloading(const FString& PluginName, const FGameFeaturePluginIdentifier& PluginIdentifier) override;
@@ -75,6 +84,8 @@ private:
 	TSharedPtr<IInstallBundleManager> BundleManager;
 	TSet<FString> RelevantGFPs;
 	TMap<FString, FGFPInstall> ActiveGFPInstalls;
+	TOptional<FInstallBundleCombinedProgressTracker> TotalProgressTracker;
+	static const FName GFOContentRequestName;
 
 	/** Delegate handle for a console variable sink */
 	FConsoleVariableSinkHandle CVarSinkHandle;
