[Optional Content Download] Minor fixes in the optional progress code

#rnx

diff --git a/Engine/Plugins/Runtime/GameFeatures/Source/GameFeatures/Private/GameFeatureOptionalContentInstaller.cpp b/Engine/Plugins/Runtime/GameFeatures/Source/GameFeatures/Private/GameFeatureOptionalContentInstaller.cpp
--- a/Engine/Plugins/Runtime/GameFeatures/Source/GameFeatures/Private/GameFeatureOptionalContentInstaller.cpp
+++ b/Engine/Plugins/Runtime/GameFeatures/Source/GameFeatures/Private/GameFeatureOptionalContentInstaller.cpp
@@ -112,8 +112,8 @@ float UGameFeatureOptionalContentInstaller::GetAllInstallsProgress()
 		StartTotalProgressTracker();
 	}
 
-	// Return 1 if some optional bundles are installed, 0 if none are installed 
-	return RelevantGFPs.Num() > 0 ? 1.f : 0.f;
+	// Return 1 if some optional bundles are installed, 0 if none are installed or active installs are present 
+	return ActiveGFPInstalls.Num() == 0 && RelevantGFPs.Num() > 0 ? 1.f : 0.f;
 }
 
 bool UGameFeatureOptionalContentInstaller::UpdateContent(const FString& PluginName, bool bIsPredownload)
@@ -172,12 +172,17 @@ bool UGameFeatureOptionalContentInstaller::UpdateContent(const FString& PluginNa
 	FInstallBundleRequestInfo& Request =  MaybeRequest.GetValue();
 	if (!Request.BundlesEnqueued.IsEmpty())
 	{
-		if (ActiveGFPInstalls.Num() == 0)
-    	{
-    		OnOptionalContentInstallStarted.Broadcast();
-    	}
+		const bool bIsOptionalContentInstallStart = ActiveGFPInstalls.Num() == 0;
 	
 		FGFPInstall& Pending = ActiveGFPInstalls.FindOrAdd(PluginName);
+		
+		if (bIsOptionalContentInstallStart)
+		{
+			// We call the delegate after adding the entry to 'ActiveGFPInstalls'. 
+			// If we called it before then the code triggered by this delegate could request information from the current script
+			// and since 'ActiveGFPInstalls' would be empty it would behave as if no installs were happening.
+			OnOptionalContentInstallStarted.Broadcast();
+		}
 
 		if (!Pending.CallbackHandle.IsValid())
 		{
