- Changed how I reordered the ini loading in GFP mounting to work with dynamic GFPs

diff --git a/Engine/Plugins/Runtime/GameFeatures/Source/GameFeatures/Private/GameFeaturePluginStateMachine.cpp b/Engine/Plugins/Runtime/GameFeatures/Source/GameFeatures/Private/GameFeaturePluginStateMachine.cpp
--- a/Engine/Plugins/Runtime/GameFeatures/Source/GameFeatures/Private/GameFeaturePluginStateMachine.cpp
+++ b/Engine/Plugins/Runtime/GameFeatures/Source/GameFeatures/Private/GameFeaturePluginStateMachine.cpp
@@ -2411,7 +2411,7 @@ struct FGameFeaturePluginState_Mounting : public FGameFeaturePluginState
 		}
 	}
 
-	void UpdateState_MountPlugin()
+	void UpdateState_MountPlugin(bool bLoadPluginIniHierarchy)
 	{
 		if (EnumHasAnyFlags(StartedSubStates, ESubState::MountPlugin))
 		{
@@ -2464,6 +2464,13 @@ struct FGameFeaturePluginState_Mounting : public FGameFeaturePluginState
 			}
 		}
 
+		// now load ini files if desired, now that we know the plugin has been loaded
+		if (bLoadPluginIniHierarchy)
+		{
+			UGameFeatureData::InitializeBasePluginIniFile(StateProperties.PluginInstalledFilename);
+		}
+
+
 		if (Result.HasError() || !bNeedsPluginMount)
 		{
 			CompletedSubStates |= ESubState::MountPlugin;
@@ -2668,7 +2675,7 @@ struct FGameFeaturePluginState_Mounting : public FGameFeaturePluginState
 
 		TRACE_CPUPROFILER_EVENT_SCOPE(GFP_Mounting);
 		
-		UpdateState_MountPlugin();
+		UpdateState_MountPlugin(AllowIniLoading());
 		UpdateState_LoadAssetRegistry();
 
 		const bool bComplete = EnumHasAllFlags(CompletedSubStates, ESubState::MountPlugin | ESubState::LoadAssetRegistry);
@@ -2676,11 +2683,6 @@ struct FGameFeaturePluginState_Mounting : public FGameFeaturePluginState
 		// Post-mount
 		if (bComplete)
 		{
-			if (AllowIniLoading())
-			{
-				UGameFeatureData::InitializeBasePluginIniFile(StateProperties.PluginInstalledFilename);
-			}
-
 			FGameFeaturePostMountingContext Context(StateProperties.PluginName, [this](FStringView InPauserTag) { OnPostMountPauserCompleted(InPauserTag); });
 			NumExpectedPostMountPausers = INDEX_NONE;
 			UGameFeaturesSubsystem::Get().OnGameFeaturePostMounting(StateProperties.PluginName, StateProperties.PluginIdentifier, Context);
diff --git a/Engine/Source/Runtime/Core/Private/Misc/ConfigCacheIni.cpp b/Engine/Source/Runtime/Core/Private/Misc/ConfigCacheIni.cpp
--- a/Engine/Source/Runtime/Core/Private/Misc/ConfigCacheIni.cpp
+++ b/Engine/Source/Runtime/Core/Private/Misc/ConfigCacheIni.cpp
@@ -2871,6 +2871,8 @@ void FConfigFile::SaveSourceToBackupFile()
 
 void FConfigFile::ProcessSourceAndCheckAgainstBackup()
 {
+	TRACE_CPUPROFILER_EVENT_SCOPE(ProcessSourceAndCheckAgainstBackup);
+
 	if (!FPlatformProperties::RequiresCookedData())
 	{
 		FString BetweenRunsDir = (FPaths::ProjectIntermediateDir() / TEXT("Config/CoalescedSourceConfigs/"));
diff --git a/Engine/Source/Runtime/Core/Private/Misc/ConfigContext.cpp b/Engine/Source/Runtime/Core/Private/Misc/ConfigContext.cpp
--- a/Engine/Source/Runtime/Core/Private/Misc/ConfigContext.cpp
+++ b/Engine/Source/Runtime/Core/Private/Misc/ConfigContext.cpp
@@ -417,6 +417,14 @@ bool FConfigContext::PrepareForLoad(bool& bPerformLoad)
 	// assume we will load, unless some code below determines not to
 	bPerformLoad = true;
 
+	// leaving the ability to go back to writing in case we actually find issues - they shouldn't be needed anymore 
+	// with all of the FConfigBranch changes
+	static bool bAllowWriteDuringLoad = FParse::Param(FCommandLine::Get(), TEXT("writeIniOnLoad"));
+	if (!bAllowWriteDuringLoad)
+	{
+		bWriteDestIni = false;
+	}
+
 	// first, make sure the DestIniFilename is set, if needed
 	if (bWriteDestIni || bAllowGeneratedIniWhenCooked || FPlatformProperties::RequiresCookedData())
 	{
