FGameFeaturePluginState_Loading
- Don't bind FStreamableHandle compeletion callback after the handle is created.  The asset manager may have bound it internally which could cause it to be stomped. See jira for details.

diff --git a/Engine/Plugins/Runtime/GameFeatures/Source/GameFeatures/Private/GameFeaturePluginStateMachine.cpp b/Engine/Plugins/Runtime/GameFeatures/Source/GameFeatures/Private/GameFeaturePluginStateMachine.cpp
--- a/Engine/Plugins/Runtime/GameFeatures/Source/GameFeatures/Private/GameFeaturePluginStateMachine.cpp
+++ b/Engine/Plugins/Runtime/GameFeatures/Source/GameFeatures/Private/GameFeaturePluginStateMachine.cpp
@@ -3399,19 +3399,6 @@ struct FGameFeaturePluginState_Loading : public FGameFeaturePluginState
 		IPluginManager::Get().MountExplicitlyLoadedPluginLocalizationData(StateProperties.PluginName);
 
 		BundleHandle = LoadGameFeatureBundles(StateProperties.GameFeatureData);
-		if (BundleHandle)
-		{
-			// This will only bind if a load is in progress
-			BundleHandle->BindCancelDelegate(FStreamableDelegate::CreateLambda([this]
-			{
-				const FStringView ShortUrl = StateProperties.PluginIdentifier.GetIdentifyingString();
-				UE_LOG(LogGameFeatures, Error, TEXT("Game Feature preloading was cancelled for URL %.*s"), ShortUrl.Len(), ShortUrl.GetData());
-				UpdateStateMachineDeferred();
-			}));
-
-			// This will only bind if a load is in progress
-			BundleHandle->BindCompleteDelegate(FStreamableDelegate::CreateRaw(this, &FGameFeaturePluginState_Loading::UpdateStateMachineDeferred, 0.0f));
-		}
 	}
 
 	virtual void EndState() override
@@ -3467,7 +3454,21 @@ struct FGameFeaturePluginState_Loading : public FGameFeaturePluginState
 		TSharedPtr<FStreamableHandle> RetHandle;
 		if (AssetIdsToLoad.Num() > 0)
 		{
-			RetHandle = UAssetManager::Get().LoadPrimaryAssets(AssetIdsToLoad, Policy.GetPreloadBundleStateForGameFeature());
+			FAssetManagerLoadParams LoadParams;
+			LoadParams.OnCancel = FStreamableDelegateWithHandle::CreateLambda([this](const TSharedPtr<FStreamableHandle>&)
+			{
+				const FStringView ShortUrl = StateProperties.PluginIdentifier.GetIdentifyingString();
+				UE_LOG(LogGameFeatures, Error, TEXT("Game Feature preloading was cancelled for URL %.*s"), ShortUrl.Len(), ShortUrl.GetData());
+				UpdateStateMachineDeferred();
+			});
+
+			// This can't be bound to the handle after its created, AM may bind it internally
+			LoadParams.OnComplete = FStreamableDelegateWithHandle::CreateLambda([this](const TSharedPtr<FStreamableHandle>&)
+			{
+				UpdateStateMachineDeferred();
+			});
+
+			RetHandle = UAssetManager::Get().LoadPrimaryAssets(AssetIdsToLoad, Policy.GetPreloadBundleStateForGameFeature(), MoveTemp(LoadParams));
 		}
 
 		return RetHandle;
