Exposed an AddSource function to the DataRegistrySource GFA for use by an internal tool

#rnx

diff --git a/Engine/Plugins/Runtime/GameFeatures/Source/GameFeatures/Private/GameFeatureAction_DataRegistrySource.cpp b/Engine/Plugins/Runtime/GameFeatures/Source/GameFeatures/Private/GameFeatureAction_DataRegistrySource.cpp
--- a/Engine/Plugins/Runtime/GameFeatures/Source/GameFeatures/Private/GameFeatureAction_DataRegistrySource.cpp
+++ b/Engine/Plugins/Runtime/GameFeatures/Source/GameFeatures/Private/GameFeatureAction_DataRegistrySource.cpp
@@ -287,5 +287,12 @@ EDataValidationResult UGameFeatureAction_DataRegistrySource::IsDataValid(FDataVa
 }
 #endif // WITH_EDITOR
 
+#if WITH_EDITOR
+void UGameFeatureAction_DataRegistrySource::AddSource(const FDataRegistrySourceToAdd& NewSource)
+{
+	SourcesToAdd.Add(NewSource);
+}
+#endif // WITH_EDITOR
+
 #undef LOCTEXT_NAMESPACE
 
diff --git a/Engine/Plugins/Runtime/GameFeatures/Source/GameFeatures/Public/GameFeatureAction_DataRegistrySource.h b/Engine/Plugins/Runtime/GameFeatures/Source/GameFeatures/Public/GameFeatureAction_DataRegistrySource.h
--- a/Engine/Plugins/Runtime/GameFeatures/Source/GameFeatures/Public/GameFeatureAction_DataRegistrySource.h
+++ b/Engine/Plugins/Runtime/GameFeatures/Source/GameFeatures/Public/GameFeatureAction_DataRegistrySource.h
@@ -70,6 +70,11 @@ public:
 #endif
 	//~End of UObject interface
 
+#if WITH_EDITOR
+	/** Used by an editor tool to programmatically register a new source */
+	GAMEFEATURES_API void AddSource(const FDataRegistrySourceToAdd& NewSource);
+#endif
+
 private:
 	/** List of sources to add when this feature is activated */
 	UPROPERTY(EditAnywhere, Category = "Registry Data", meta=(TitleProperty="RegistryToAddTo"))
