Remove experimental GFP subsystem IAD support. The AssetManager and StreamableManager now support IAD.

diff --git a/Engine/Plugins/Runtime/GameFeatures/Source/GameFeatures/Private/GameFeaturesSubsystem.cpp b/Engine/Plugins/Runtime/GameFeatures/Source/GameFeatures/Private/GameFeaturesSubsystem.cpp
--- a/Engine/Plugins/Runtime/GameFeatures/Source/GameFeatures/Private/GameFeaturesSubsystem.cpp
+++ b/Engine/Plugins/Runtime/GameFeatures/Source/GameFeatures/Private/GameFeaturesSubsystem.cpp
@@ -2867,395 +2867,6 @@ TSharedRef<FGameFeaturePluginPredownloadHandle> UGameFeaturesSubsystem::Predownl
 	return Context;
 }
 
-FGameFeaturePluginStreamedAssetsRef::FGameFeaturePluginStreamedAssetsRef(FStringView DebugName)
-	: ContentHandle(UE::IoStore::FOnDemandContentHandle::Create(DebugName))
-{
-}
-
-FString FGameFeaturePluginStreamedAssetsRef::GetDebugName() const
-{
-	return LexToString(ContentHandle);
-}
-
-void FGameFeaturePluginStreamedAssetsRef::Release()
-{
-	ContentHandle.Reset();
-}
-
-bool FGameFeaturePluginStreamedAssetsRef::IsValid() const
-{
-	return ContentHandle.IsValid();
-}
-
-struct FGameFeatureAssetStreamContext : public FGameFeaturePluginAssetStreamHandle
-{
-	const FStringView ErrorNamespace = TEXTVIEW("GameFeaturePlugin.AssetStream.");
-	FGameFeaturePluginStreamedAssetsRef AssetRef;
-	TUniqueFunction<void(const UE::GameFeatures::FResult&)> OnComplete;
-	TUniqueFunction<void(float)> OnProgress;
-	UE::GameFeatures::FResult Result = MakeValue();
-	TArray<UE::IoStore::FOnDemandInstallProgress> ProgressList;
-	TArray<UE::IoStore::FOnDemandInstallRequest> InstallRequests;
-	int32 PendingInstalls = 0;
-	bool bIsComplete = false;
-	bool bIsCanceled = false;
-
-	FGameFeatureAssetStreamContext(FGameFeaturePluginStreamedAssetsRef&& InAssetRef)
-		: AssetRef(MoveTemp(InAssetRef))
-	{}
-
-	virtual bool IsComplete() const override
-	{
-		return bIsComplete;
-	}
-
-	virtual const UE::GameFeatures::FResult& GetResult() const override
-	{
-		return Result;
-	}
-
-	virtual float GetProgress() const override
-	{
-		const UE::IoStore::FOnDemandInstallProgress SumProgress = Algo::Accumulate(
-			ProgressList, UE::IoStore::FOnDemandInstallProgress(), &UE::IoStore::FOnDemandInstallProgress::Combine);
-		return SumProgress.GetRelativeProgress();
-	}
-
-	virtual void Cancel() override
-	{
-		bIsCanceled = true;
-		for (UE::IoStore::FOnDemandInstallRequest& Request : InstallRequests)
-		{
-			Request.Cancel();
-		}
-	}
-
-	virtual FGameFeaturePluginStreamedAssetsRef GetAssetsRef() const override
-	{
-		return AssetRef;
-	}
-
-	bool IsCanceled()
-	{
-		return bIsCanceled;
-	}
-
-	void SetComplete()
-	{
-		bIsComplete = true;
-		if (OnComplete)
-		{
-			OnComplete(Result);
-		}
-	}
-
-	void SetComplete(UE::GameFeatures::FResult&& InResult)
-	{
-		Result = MoveTemp(InResult);
-		SetComplete();
-	}
-
-	void SetCompleteCanceled()
-	{
-		Result = MakeError(FString::Printf(TEXT("%.*s%s"),
-			ErrorNamespace.Len(), ErrorNamespace.GetData(),
-			TEXT("Canceled")));
-		SetComplete();
-	}
-
-	void Start(TArrayView<FName> AssetPackageNames, EGameFeatureAssetStreamPriority Priority)
-	{
-		if (IsCanceled())
-		{
-			SetCompleteCanceled();
-			return;
-		}
-
-		UGameFeaturesSubsystem& GFPSubSys = UGameFeaturesSubsystem::Get();
-
-		TSet<FStringView> MountPoints;
-
-		// Find mount points
-		FNameBuilder PackageNameBuilder;
-		for (FName PackageName : AssetPackageNames)
-		{
-			PackageName.ToString(PackageNameBuilder);
-			const FStringView MountPoint = FPathViews::GetMountPointNameFromPath(PackageNameBuilder);
-			MountPoints.Add(MountPoint);
-		}
-
-		// Verify that the we can stream from these GFPs
-		for (FStringView MountPoint : MountPoints)
-		{
-			FString PluginURL;
-			if (!GFPSubSys.GetPluginURLByName(MountPoint, PluginURL))
-			{
-				UE_LOGFMT(LogGameFeatures, Error, "{MountPoint} is unknown", ("MountPoint", MountPoint));
-				UE::GameFeatures::FResult ErrorResult = MakeError(FString::Printf(TEXT("%.*s%s"),
-					ErrorNamespace.Len(), ErrorNamespace.GetData(), TEXT("Unknown_Mount_Point")));
-				SetComplete(MoveTemp(ErrorResult));
-				return;
-			}
-
-			const UGameFeaturePluginStateMachine* Machine = GFPSubSys.FindGameFeaturePluginStateMachine(PluginURL);
-			if (!Machine)
-			{
-				UE_LOGFMT(LogGameFeatures, Error, "Could not find a GFP with mount point {MountPoint}", ("MountPoint", MountPoint));
-				UE::GameFeatures::FResult ErrorResult = MakeError(FString::Printf(TEXT("%.*s%s"),
-					ErrorNamespace.Len(), ErrorNamespace.GetData(), TEXT("Unknown_GFP_Mount_Point")));
-				SetComplete(MoveTemp(ErrorResult));
-				return;
-			}
-
-			if (Machine->GetPluginIdentifier().GetPluginProtocol() != EGameFeaturePluginProtocol::InstallBundle)
-			{
-				UE_LOGFMT(LogGameFeatures, Error, "Unsupported protocol for GFP {MountPoint}", ("MountPoint", MountPoint));
-				UE::GameFeatures::FResult ErrorResult = MakeError(FString::Printf(TEXT("%.*s%s"),
-					ErrorNamespace.Len(), ErrorNamespace.GetData(), TEXT("Unsupported_GFP_Protocol")));
-				SetComplete(MoveTemp(ErrorResult));
-				return;
-			}
-
-			// Check installed and not moving to a state less than installed
-			if (Machine->GetCurrentState() < EGameFeaturePluginState::Installed ||
-				Machine->GetDestination() < EGameFeaturePluginState::Installed)
-			{
-				UE_LOGFMT(LogGameFeatures, Error, "Invalid GFP state for {MountPoint}", ("MountPoint", MountPoint));
-				UE::GameFeatures::FResult ErrorResult = MakeError(FString::Printf(TEXT("%.*s%s"),
-					ErrorNamespace.Len(), ErrorNamespace.GetData(), TEXT("GFP_Not_Ready")));
-				SetComplete(MoveTemp(ErrorResult));
-				return;
-			}
-		}
-
-		// Install the packages
-		UE::IoStore::FOnDemandInstallArgs InstallArgs;
-		InstallArgs.PackageIds.Reserve(AssetPackageNames.Num());
-		Algo::Transform(AssetPackageNames, InstallArgs.PackageIds, &FPackageId::FromName);
-
-		Start(MakeArrayView(&InstallArgs, 1), Priority);
-	}
-
-	void StartByTag(TArrayView<FString> Tags, EGameFeatureAssetStreamPriority Priority)
-	{
-		if (IsCanceled())
-		{
-			SetCompleteCanceled();
-			return;
-		}
-
-		// This path streams matching packages from all mounted plugins so plugin state can't
-		// be verfied here. The user will just have to deal with low level errors if they occur.
-
-		UE::IoStore::FOnDemandInstallArgs InstallArgs;
-		InstallArgs.TagSets = Tags;
-
-		Start(MakeArrayView(&InstallArgs, 1), Priority);
-	}
-
-	void StartByTag(TArrayView<FString> PluginURLs, TArrayView<FString> Tags, EGameFeatureAssetStreamPriority Priority)
-	{
-		if (IsCanceled())
-		{
-			SetCompleteCanceled();
-			return;
-		}
-
-		UGameFeaturesSubsystem& GFPSubSys = UGameFeaturesSubsystem::Get();
-
-		TArray<FName> ContainerMountIds;
-
-		// Verify that the we can stream from these GFPs
-		for (const FString& PluginURL : PluginURLs)
-		{
-			const UGameFeaturePluginStateMachine* Machine = GFPSubSys.FindGameFeaturePluginStateMachine(PluginURL);
-			if (!Machine)
-			{
-				FStringView PluginPath;
-				UGameFeaturesSubsystem::ParsePluginURL(PluginURL, nullptr, &PluginPath);
-
-				UE_LOGFMT(LogGameFeatures, Error, "Could not find GFP {PluginPath}", ("PluginPath", PluginPath));
-				UE::GameFeatures::FResult ErrorResult = MakeError(FString::Printf(TEXT("%.*s%s"),
-					ErrorNamespace.Len(), ErrorNamespace.GetData(), TEXT("Unknown_GFP_Mount_Point")));
-				SetComplete(MoveTemp(ErrorResult));
-				return;
-			}
-
-			FStringView PluginName = Machine->GetPluginIdentifier().GetPluginName();
-
-			if (Machine->GetPluginIdentifier().GetPluginProtocol() != EGameFeaturePluginProtocol::InstallBundle)
-			{
-				UE_LOGFMT(LogGameFeatures, Error, "Unsupported protocol for GFP {PluginName}", ("PluginName", PluginName));
-				UE::GameFeatures::FResult ErrorResult = MakeError(FString::Printf(TEXT("%.*s%s"),
-					ErrorNamespace.Len(), ErrorNamespace.GetData(), TEXT("Unsupported_GFP_Protocol")));
-				SetComplete(MoveTemp(ErrorResult));
-				return;
-			}
-
-			// Check installed and not moving to a state less than installed
-			if (Machine->GetCurrentState() >= EGameFeaturePluginState::Installed &&
-				Machine->GetDestination() > EGameFeaturePluginState::StatusKnown)
-			{
-				UE_LOGFMT(LogGameFeatures, Error, "Invalid GFP state for {PluginName}", ("PluginName", PluginName));
-				UE::GameFeatures::FResult ErrorResult = MakeError(FString::Printf(TEXT("%.*s%s"),
-					ErrorNamespace.Len(), ErrorNamespace.GetData(), TEXT("GFP_Not_Ready")));
-				SetComplete(MoveTemp(ErrorResult));
-				return;
-			}
-
-			// Get mount ids of containers to install assets from
-			const FInstallBundlePluginProtocolMetaData& Metadata = Machine->GetProtocolMetadata().GetSubtype<FInstallBundlePluginProtocolMetaData>();
-			for (const FName InstallBundle : Metadata.InstallBundlesWithAssetDependencies)
-			{
-				ContainerMountIds.AddUnique(InstallBundle);
-			}
-		}
-
-		if (ContainerMountIds.IsEmpty())
-		{
-			SetComplete(); // Success?! There's nothing to install
-		}
-
-		TArray<UE::IoStore::FOnDemandInstallArgs> InstallArgsList;
-		InstallArgsList.Reserve(ContainerMountIds.Num());
-		for (const FName MountId : ContainerMountIds)
-		{
-			UE::IoStore::FOnDemandInstallArgs& InstallArgs = InstallArgsList.Emplace_GetRef();
-			InstallArgs.MountId = MountId.ToString();
-			InstallArgs.TagSets = Tags;
-		}
-
-		Start(InstallArgsList, Priority);
-	}
-
-	void Start(TArrayView<UE::IoStore::FOnDemandInstallArgs> InstallArgList, EGameFeatureAssetStreamPriority Priority)
-	{
-		UE::IoStore::IOnDemandIoStore* IoStore = UE::IoStore::TryGetOnDemandIoStore();
-		if (!IoStore)
-		{
-			UE::GameFeatures::FResult ErrorResult = MakeError(FString::Printf(TEXT("%.*s%s"),
-				ErrorNamespace.Len(), ErrorNamespace.GetData(), TEXT("IoStoreOnDemand.ModuleNotFound")));
-			SetComplete(MoveTemp(ErrorResult));
-			return;
-		}
-
-		if (!AssetRef.IsValid())
-		{
-			UE::GameFeatures::FResult ErrorResult = MakeError(FString::Printf(TEXT("%.*s%s"),
-				ErrorNamespace.Len(), ErrorNamespace.GetData(), TEXT("Invalid_AssetRef")));
-			SetComplete(MoveTemp(ErrorResult));
-			return;
-		}
-
-		check(InstallArgList.Num() > 0);
-
-		PendingInstalls = InstallArgList.Num();
-		ProgressList.Empty(PendingInstalls);
-		InstallRequests.Empty(PendingInstalls);
-
-		for (UE::IoStore::FOnDemandInstallArgs& InstallArgs : InstallArgList)
-		{
-			const int32 RequestIndex = ProgressList.Emplace();
-
-			InstallArgs.ContentHandle = AssetRef.ContentHandle;
-			InstallArgs.Options |= UE::IoStore::EOnDemandInstallOptions::CallbackOnGameThread;
-			InstallArgs.Options |= UE::IoStore::EOnDemandInstallOptions::InstallSoftReferences;
-
-			// TODO: convert Priority to IOStore priority here
-
-			InstallRequests.Add(IoStore->Install(MoveTemp(InstallArgs),
-				[Context = SharedThis(this), RequestIndex](const UE::IoStore::FOnDemandInstallResult& OnDemandInstallResult)
-				{
-					Context->OnInstallComplete(OnDemandInstallResult, RequestIndex);
-				},
-				[Context = SharedThis(this), RequestIndex](const UE::IoStore::FOnDemandInstallProgress& InProgress)
-				{
-					Context->OnInstallProgress(InProgress, RequestIndex);
-				}));
-		}
-	}
-
-	void OnInstallComplete(const UE::IoStore::FOnDemandInstallResult& OnDemandInstallResult, const int32 RequestIndex)
-	{
-		if (!OnDemandInstallResult.Status.IsOk() && !Result.HasError())
-		{
-			const FString StatusString = OnDemandInstallResult.Status.ToString();
-			FString ErrorCode = GetIoErrorText(OnDemandInstallResult.Status.GetErrorCode());
-			ErrorCode.ReplaceCharInline(TEXT(' '), TEXT('_'), ESearchCase::CaseSensitive);
-
-			Result = UE::GameFeatures::FResult(MakeError(
-				FString::Printf(TEXT("%.*s%s.%s"),
-					ErrorNamespace.Len(), ErrorNamespace.GetData(), TEXT("IoStoreOnDemand.OnComplete"), *ErrorCode)),
-				FText::AsCultureInvariant(StatusString));
-
-			// Something went wrong, kill any other requests
-			Cancel();
-		}
-
-		--PendingInstalls;
-		if (PendingInstalls == 0)
-		{
-			SetComplete();
-		}
-	}
-
-	void OnInstallProgress(const UE::IoStore::FOnDemandInstallProgress& InProgress, const int32 RequestIndex)
-	{
-		ProgressList[RequestIndex] = InProgress;
-		if (OnProgress)
-		{
-			OnProgress(GetProgress());
-		}
-	}
-};
-
-TSharedRef<FGameFeaturePluginAssetStreamHandle> UGameFeaturesSubsystem::DownloadStreamedAssets(
-	FGameFeaturePluginStreamedAssetsRef AssetRef,
-	TArrayView<FName> AssetPackageNames,
-	EGameFeatureAssetStreamPriority Priority /*= EGameFeatureAssetStreamPriority::Normal*/,
-	TUniqueFunction<void(const UE::GameFeatures::FResult&)> OnComplete /*= nullptr*/, 
-	TUniqueFunction<void(float)> OnProgress /*= nullptr*/)
-{
-	TSharedRef<FGameFeatureAssetStreamContext> Context = MakeShared<FGameFeatureAssetStreamContext>(MoveTemp(AssetRef));
-	Context->OnComplete = MoveTemp(OnComplete);
-	Context->OnProgress = MoveTemp(OnProgress);
-	Context->Start(AssetPackageNames, Priority);
-
-	return Context;
-}
-
-TSharedRef<FGameFeaturePluginAssetStreamHandle> UGameFeaturesSubsystem::DownloadStreamedAssetsByTag(
-	FGameFeaturePluginStreamedAssetsRef AssetRef,
-	TArrayView<FString> Tags,
-	EGameFeatureAssetStreamPriority Priority /*= EGameFeatureAssetStreamPriority::Normal*/,
-	TUniqueFunction<void(const UE::GameFeatures::FResult&)> OnComplete /*= nullptr*/,
-	TUniqueFunction<void(float)> OnProgress /*= nullptr*/)
-{
-	// TODO: not sure if the overload makes a lot of sense. What's the use case?
-
-	TSharedRef<FGameFeatureAssetStreamContext> Context = MakeShared<FGameFeatureAssetStreamContext>(MoveTemp(AssetRef));
-	Context->OnComplete = MoveTemp(OnComplete);
-	Context->OnProgress = MoveTemp(OnProgress);
-	Context->StartByTag(Tags, Priority);
-
-	return Context;
-}
-
-TSharedRef<FGameFeaturePluginAssetStreamHandle> DownloadStreamedAssetsByTag(
-	FGameFeaturePluginStreamedAssetsRef AssetRef,
-	TArrayView<FString> PluginURLs,
-	TArrayView<FString> Tags,
-	EGameFeatureAssetStreamPriority Priority /*= EGameFeatureAssetStreamPriority::Normal*/,
-	TUniqueFunction<void(const UE::GameFeatures::FResult&)> OnComplete /*= nullptr*/,
-	TUniqueFunction<void(float)> OnProgress /*= nullptr*/)
-{
-	TSharedRef<FGameFeatureAssetStreamContext> Context = MakeShared<FGameFeatureAssetStreamContext>(MoveTemp(AssetRef));
-	Context->OnComplete = MoveTemp(OnComplete);
-	Context->OnProgress = MoveTemp(OnProgress);
-	Context->StartByTag(PluginURLs, Tags, Priority);
-	return Context;
-}
-
 UGameFeaturePluginStateMachine* UGameFeaturesSubsystem::FindGameFeaturePluginStateMachine(const FString& PluginURL) const
 {
 	FGameFeaturePluginIdentifier FindPluginIdentifier(PluginURL);
diff --git a/Engine/Plugins/Runtime/GameFeatures/Source/GameFeatures/Public/GameFeaturesSubsystem.h b/Engine/Plugins/Runtime/GameFeatures/Source/GameFeatures/Public/GameFeaturesSubsystem.h
--- a/Engine/Plugins/Runtime/GameFeatures/Source/GameFeatures/Public/GameFeaturesSubsystem.h
+++ b/Engine/Plugins/Runtime/GameFeatures/Source/GameFeatures/Public/GameFeaturesSubsystem.h
@@ -286,38 +286,6 @@ struct FGameFeaturePluginPredownloadHandle : public FGameFeaturePluginAsyncHandl
 {
 };
 
-/** Referencer to hold individual asset downloads in local cache */
-struct GAMEFEATURES_API FGameFeaturePluginStreamedAssetsRef
-{
-	explicit FGameFeaturePluginStreamedAssetsRef(FStringView DebugName);
-
-	FString GetDebugName() const;
-
-	void Release();
-	bool IsValid() const;
-	operator bool() const { return IsValid(); }
-
-private:
-	UE::IoStore::FOnDemandContentHandle ContentHandle;
-	friend struct FGameFeatureAssetStreamContext;
-};
-
-/** Handle to track downloads of individual assets */
-struct FGameFeaturePluginAssetStreamHandle : public FGameFeaturePluginAsyncHandle
-{
-	virtual FGameFeaturePluginStreamedAssetsRef GetAssetsRef() const = 0;
-};
-
-UENUM(BlueprintType)
-enum class EGameFeatureAssetStreamPriority : uint8
-{
-	Lowest = 0,
-	BelowNormal = 127,
-	Normal = 128,
-	AboveNormal = 129,
-	Highest = 255
-};
-
 /** Struct used to transform a GameFeaturePlugin URL into something that can uniquely identify the GameFeaturePlugin
     without including any transient data being passed in through the URL */
 USTRUCT()
@@ -712,32 +680,6 @@ public:
 		TUniqueFunction<void(const UE::GameFeatures::FResult&)> OnComplete = nullptr, TUniqueFunction<void(float)> OnProgress = nullptr);
 	friend struct FGameFeaturePluginPredownloadContext;
 
-	/**
-	 * ***************EXPERIMENTAL***************** 
-	 * Download streamed assets from installed GFPs. The assets cannot be loaded unless the plugin is also mounted.
-	 * Assets are gauranteed to be available as long as AssetRef is valid.
-	 */
-	TSharedRef<FGameFeaturePluginAssetStreamHandle> DownloadStreamedAssets(
-		FGameFeaturePluginStreamedAssetsRef AssetRef, 
-		TArrayView<FName> AssetPackageNames,
-		EGameFeatureAssetStreamPriority Priority = EGameFeatureAssetStreamPriority::Normal,
-		TUniqueFunction<void(const UE::GameFeatures::FResult&)> OnComplete = nullptr, 
-		TUniqueFunction<void(float)> OnProgress = nullptr);
-	TSharedRef<FGameFeaturePluginAssetStreamHandle> DownloadStreamedAssetsByTag(
-		FGameFeaturePluginStreamedAssetsRef AssetRef,
-		TArrayView<FString> Tags,
-		EGameFeatureAssetStreamPriority Priority = EGameFeatureAssetStreamPriority::Normal,
-		TUniqueFunction<void(const UE::GameFeatures::FResult&)> OnComplete = nullptr,
-		TUniqueFunction<void(float)> OnProgress = nullptr);
-	TSharedRef<FGameFeaturePluginAssetStreamHandle> DownloadStreamedAssetsByTag(
-		FGameFeaturePluginStreamedAssetsRef AssetRef,
-		TArrayView<FString> PluginURLs,
-		TArrayView<FString> Tags,
-		EGameFeatureAssetStreamPriority Priority = EGameFeatureAssetStreamPriority::Normal,
-		TUniqueFunction<void(const UE::GameFeatures::FResult&)> OnComplete = nullptr,
-		TUniqueFunction<void(float)> OnProgress = nullptr);
-	friend struct FGameFeatureAssetStreamContext;
-
 	/** Determine the initial feature state for a built-in plugin */
 	static EBuiltInAutoState DetermineBuiltInInitialFeatureState(TSharedPtr<FJsonObject> Descriptor, const FString& ErrorContext);
 
