Improved batching for compilation during plugin loading.

#rnx

diff --git a/Engine/Plugins/Runtime/GameFeatures/Source/GameFeatures/Private/GameFeaturesSubsystem.cpp b/Engine/Plugins/Runtime/GameFeatures/Source/GameFeatures/Private/GameFeaturesSubsystem.cpp
--- a/Engine/Plugins/Runtime/GameFeatures/Source/GameFeatures/Private/GameFeaturesSubsystem.cpp
+++ b/Engine/Plugins/Runtime/GameFeatures/Source/GameFeatures/Private/GameFeaturesSubsystem.cpp
@@ -2158,6 +2158,8 @@ private:
 
 void UGameFeaturesSubsystem::LoadBuiltInGameFeaturePlugins(FBuiltInPluginAdditionalFilters AdditionalFilter, const FBuiltInGameFeaturePluginsLoaded& InCompleteDelegate /*= FBuiltInGameFeaturePluginsLoaded()*/)
 {
+	TRACE_CPUPROFILER_EVENT_SCOPE_STR("UGameFeaturesSubsystem::LoadBuiltInGameFeaturePlugins");
+
 	struct FLoadContext
 	{
 		FScopeLogTime ScopeLogTime{TEXT("BuiltInGameFeaturePlugins loaded."), nullptr, FConditionalScopeLogTime::ScopeLog_Seconds};
