[GameFeatures] Add a way to attach a breakpoint to specific game feature's state changes
* Can be enabled while editing a game feature or using the console command "EnableDebugGameFeatureState <PluginUrl>"

#tests Used both the UI and console command versions to debug an incorrect plugin state issue

diff --git a/Engine/Plugins/Runtime/GameFeatures/Source/GameFeatures/Private/GameFeaturePluginStateMachine.cpp b/Engine/Plugins/Runtime/GameFeatures/Source/GameFeatures/Private/GameFeaturePluginStateMachine.cpp
--- a/Engine/Plugins/Runtime/GameFeatures/Source/GameFeatures/Private/GameFeaturePluginStateMachine.cpp
+++ b/Engine/Plugins/Runtime/GameFeatures/Source/GameFeatures/Private/GameFeaturePluginStateMachine.cpp
@@ -3899,19 +3899,18 @@ bool UGameFeaturePluginStateMachine::SetDestination(FGameFeaturePluginStateRange
 	check(IsValidDestinationState(InDestination.MinState));
 	check(IsValidDestinationState(InDestination.MaxState));
 
+	bool bDestinationSet = false;
+	bool bDestinationChanged = false;
+
 	if (!InDestination.IsValid())
 	{
 		// Invalid range
-		return false;
 	}
-
-	if (CurrentStateInfo.State == EGameFeaturePluginState::Terminal && !InDestination.Contains(EGameFeaturePluginState::Terminal))
+	else if (CurrentStateInfo.State == EGameFeaturePluginState::Terminal && !InDestination.Contains(EGameFeaturePluginState::Terminal))
 	{
 		// Can't tranistion away from terminal state
-		return false;
 	}
-
-	if (!IsRunning())
+	else if (!IsRunning())
 	{
 		// Not running so any new range is acceptable
 
@@ -3925,35 +3924,36 @@ bool UGameFeaturePluginStateMachine::SetDestination(FGameFeaturePluginStateRange
 		if (InDestination.Contains(CurrentStateInfo.State))
 		{
 			OnFeatureStateTransitionComplete.ExecuteIfBound(this, MakeValue());
-			return true;
 		}
-		
-		if (CurrentStateInfo.State < InDestination)
+		else
 		{
-			FDestinationGameFeaturePluginState* MinDestState = AllStates[InDestination.MinState]->AsDestinationState();
-			FDelegateHandle CallbackHandle = MinDestState->OnDestinationStateReached.Add(MoveTemp(OnFeatureStateTransitionComplete));
-			if (OutCallbackHandle)
+			if (CurrentStateInfo.State < InDestination)
 			{
-				*OutCallbackHandle = CallbackHandle;
+				FDestinationGameFeaturePluginState* MinDestState = AllStates[InDestination.MinState]->AsDestinationState();
+				FDelegateHandle CallbackHandle = MinDestState->OnDestinationStateReached.Add(MoveTemp(OnFeatureStateTransitionComplete));
+				if (OutCallbackHandle)
+				{
+					*OutCallbackHandle = CallbackHandle;
+				}
 			}
-		}
-		else if (CurrentStateInfo.State > InDestination)
-		{
-			FDestinationGameFeaturePluginState* MaxDestState = AllStates[InDestination.MaxState]->AsDestinationState();
-			FDelegateHandle CallbackHandle = MaxDestState->OnDestinationStateReached.Add(MoveTemp(OnFeatureStateTransitionComplete));
-			if (OutCallbackHandle)
+			else if (CurrentStateInfo.State > InDestination)
 			{
-				*OutCallbackHandle = CallbackHandle;
+				FDestinationGameFeaturePluginState* MaxDestState = AllStates[InDestination.MaxState]->AsDestinationState();
+				FDelegateHandle CallbackHandle = MaxDestState->OnDestinationStateReached.Add(MoveTemp(OnFeatureStateTransitionComplete));
+				if (OutCallbackHandle)
+				{
+					*OutCallbackHandle = CallbackHandle;
+				}
 			}
-		}
 
-		StateProperties.Destination = InDestination;
-		UpdateStateMachine();
+			StateProperties.Destination = InDestination;
+			UpdateStateMachine();
+			bDestinationChanged = true;
+		}
 
-		return true;
+		bDestinationSet = true;
 	}
-
-	if (TOptional<FGameFeaturePluginStateRange> NewDestination = StateProperties.Destination.Intersect(InDestination))
+	else if (TOptional<FGameFeaturePluginStateRange> NewDestination = StateProperties.Destination.Intersect(InDestination))
 	{
 		// The machine is already running so we can only transition to this range if it overlaps with our current range.
 		// We can satisfy both ranges in this case.
@@ -3970,14 +3970,16 @@ bool UGameFeaturePluginStateMachine::SetDestination(FGameFeaturePluginStateRange
 			if (InDestination.Contains(CurrentStateInfo.State))
 			{
 				OnFeatureStateTransitionComplete.ExecuteIfBound(this, MakeValue());
-				return true;
 			}
-
-			FDestinationGameFeaturePluginState* MinDestState = AllStates[InDestination.MinState]->AsDestinationState();
-			FDelegateHandle CallbackHandle = MinDestState->OnDestinationStateReached.Add(MoveTemp(OnFeatureStateTransitionComplete));
-			if (OutCallbackHandle)
+			else
 			{
-				*OutCallbackHandle = CallbackHandle;
+				FDestinationGameFeaturePluginState* MinDestState = AllStates[InDestination.MinState]->AsDestinationState();
+				FDelegateHandle CallbackHandle = MinDestState->OnDestinationStateReached.Add(MoveTemp(OnFeatureStateTransitionComplete));
+				if (OutCallbackHandle)
+				{
+					*OutCallbackHandle = CallbackHandle;
+				}
+				bDestinationChanged = true;
 			}
 		}
 		else if(CurrentStateInfo.State > StateProperties.Destination)
@@ -3987,14 +3989,16 @@ bool UGameFeaturePluginStateMachine::SetDestination(FGameFeaturePluginStateRange
 			if (InDestination.Contains(CurrentStateInfo.State))
 			{
 				OnFeatureStateTransitionComplete.ExecuteIfBound(this, MakeValue());
-				return true;
 			}
-
-			FDestinationGameFeaturePluginState* MaxDestState = AllStates[InDestination.MaxState]->AsDestinationState();
-			FDelegateHandle CallbackHandle = MaxDestState->OnDestinationStateReached.Add(MoveTemp(OnFeatureStateTransitionComplete));
-			if (OutCallbackHandle)
+			else
 			{
-				*OutCallbackHandle = CallbackHandle;
+				FDestinationGameFeaturePluginState* MaxDestState = AllStates[InDestination.MaxState]->AsDestinationState();
+				FDelegateHandle CallbackHandle = MaxDestState->OnDestinationStateReached.Add(MoveTemp(OnFeatureStateTransitionComplete));
+				if (OutCallbackHandle)
+				{
+					*OutCallbackHandle = CallbackHandle;
+				}
+				bDestinationChanged = true;
 			}
 		}
 		else
@@ -4002,11 +4006,21 @@ bool UGameFeaturePluginStateMachine::SetDestination(FGameFeaturePluginStateRange
 			checkf(false, TEXT("IsRunning() returned true but state machine has reached destination!"));
 		}
 
-		return true;
+		bDestinationSet = true;
+	}
+	else
+	{
+		// The requested state range is completely outside the the current state range so reject the request
 	}
 
-	// The requested state range is completely outside the the current state range so reject the request
-	return false;
+#if !UE_BUILD_SHIPPING
+	if (bDestinationChanged && UGameFeaturesSubsystem::Get().GetPluginDebugStateEnabled(GetPluginURL()))
+	{
+		PLATFORM_BREAK();
+	}
+#endif
+
+	return bDestinationSet;
 }
 
 bool UGameFeaturePluginStateMachine::TryCancel(FGameFeatureStateTransitionCanceled OnFeatureStateTransitionCanceled, FDelegateHandle* OutCallbackHandle /*= nullptr*/)
diff --git a/Engine/Plugins/Runtime/GameFeatures/Source/GameFeatures/Private/GameFeaturesSubsystem.cpp b/Engine/Plugins/Runtime/GameFeatures/Source/GameFeatures/Private/GameFeaturesSubsystem.cpp
--- a/Engine/Plugins/Runtime/GameFeatures/Source/GameFeatures/Private/GameFeaturesSubsystem.cpp
+++ b/Engine/Plugins/Runtime/GameFeatures/Source/GameFeatures/Private/GameFeaturesSubsystem.cpp
@@ -466,6 +466,32 @@ void UGameFeaturesSubsystem::Initialize(FSubsystemCollectionBase& Collection)
 		}),
 		ECVF_Cheat);
 
+#if !UE_BUILD_SHIPPING
+	IConsoleManager::Get().RegisterConsoleCommand(
+		TEXT("EnableDebugGameFeatureState"),
+		TEXT("Trigger a debug breakpoint if when the state of the gameplay feature changes"),
+		FConsoleCommandWithWorldArgsAndOutputDeviceDelegate::CreateLambda([](const TArray<FString>& Args, UWorld*, FOutputDevice& Ar)
+			{
+				if (TOptional<FString> PluginURL = UE::GameFeatures::GetPluginUrlForConsoleCommand(Args, Ar))
+				{
+					UGameFeaturesSubsystem::Get().SetPluginDebugStateEnabled(PluginURL.GetValue(), true);
+				}
+			}),
+		ECVF_Cheat);
+
+	IConsoleManager::Get().RegisterConsoleCommand(
+		TEXT("DisableDebugGameFeatureState"),
+		TEXT("Tuen off triggering a debug breakpoint if when the state of the gameplay feature changes"),
+		FConsoleCommandWithWorldArgsAndOutputDeviceDelegate::CreateLambda([](const TArray<FString>& Args, UWorld*, FOutputDevice& Ar)
+			{
+				if (TOptional<FString> PluginURL = UE::GameFeatures::GetPluginUrlForConsoleCommand(Args, Ar))
+				{
+					UGameFeaturesSubsystem::Get().SetPluginDebugStateEnabled(PluginURL.GetValue(), false);
+				}
+			}),
+		ECVF_Cheat);
+#endif
+
 	GetExplanationForUnavailablePackageDelegateHandle = IPluginManager::Get().GetExplanationForUnavailablePackageWithPluginInfoDelegate().AddUObject(this, &UGameFeaturesSubsystem::GetExplanationForUnavailablePackage);
 	OnPluginUnmountedDelegateHandle = IPluginManager::Get().OnPluginUnmounted().AddWeakLambda(this, [this] (const IPlugin& Plugin) { PruneCachedGameFeaturePluginDetails(/* doesnt use URL */ FString(""), Plugin.GetDescriptorFileName()); });
 }
@@ -3766,6 +3792,28 @@ void UGameFeaturesSubsystem::GetPluginsToCook(TSet<FString>& OutPlugins)
 	OutPlugins.Append(PluginsList);	
 }
 
+bool UGameFeaturesSubsystem::GetPluginDebugStateEnabled(const FString& PluginUrl)
+{
+#if !UE_BUILD_SHIPPING
+	return DebugStateChangedForPlugins.Contains(PluginUrl);
+#else
+	return false;
+#endif
+}
+void UGameFeaturesSubsystem::SetPluginDebugStateEnabled(const FString& PluginUrl, bool bEnabled)
+{
+#if !UE_BUILD_SHIPPING
+	if (bEnabled)
+	{
+		DebugStateChangedForPlugins.Add(PluginUrl);
+	}
+	else
+	{
+		DebugStateChangedForPlugins.Remove(PluginUrl);
+	}
+#endif
+}
+
 void UGameFeaturesSubsystem::SetExplanationForNotMountingPlugin(const FString& PluginURL, const FString& Explanation)
 {
 #if WITH_EDITOR
diff --git a/Engine/Plugins/Runtime/GameFeatures/Source/GameFeatures/Public/GameFeaturesSubsystem.h b/Engine/Plugins/Runtime/GameFeatures/Source/GameFeatures/Public/GameFeaturesSubsystem.h
--- a/Engine/Plugins/Runtime/GameFeatures/Source/GameFeatures/Public/GameFeaturesSubsystem.h
+++ b/Engine/Plugins/Runtime/GameFeatures/Source/GameFeatures/Public/GameFeaturesSubsystem.h
@@ -687,6 +687,10 @@ public:
 
 	/** Used during a DLC cook to determine which plugins should be cooked */
 	static void GetPluginsToCook(TSet<FString>& OutPlugins);
+
+	bool GetPluginDebugStateEnabled(const FString& PluginUrl);
+	void SetPluginDebugStateEnabled(const FString& PluginUrl, bool bEnabled);
+
 private:
 	TSet<FString> GetActivePluginNames() const;
 
@@ -895,6 +899,10 @@ private:
 	// When we decide not to mount a plugin, we can store an explanation here so that if we later attempt to load an asset from it we can tell the user why it's not available
 	TMap<FString, FString> UnmountedPluginNameToExplanation;
 #endif
+
+#if !UE_BUILD_SHIPPING
+	TSet<FString> DebugStateChangedForPlugins;
+#endif
 	FDelegateHandle GetExplanationForUnavailablePackageDelegateHandle;
 	FDelegateHandle OnPluginUnmountedDelegateHandle;
 
diff --git a/Engine/Plugins/Runtime/GameFeatures/Source/GameFeaturesEditor/Private/GameFeatureDataDetailsCustomization.cpp b/Engine/Plugins/Runtime/GameFeatures/Source/GameFeaturesEditor/Private/GameFeatureDataDetailsCustomization.cpp
--- a/Engine/Plugins/Runtime/GameFeatures/Source/GameFeaturesEditor/Private/GameFeatureDataDetailsCustomization.cpp
+++ b/Engine/Plugins/Runtime/GameFeatures/Source/GameFeaturesEditor/Private/GameFeatureDataDetailsCustomization.cpp
@@ -21,6 +21,7 @@
 
 #include "GameFeatureData.h"
 #include "GameFeatureTypes.h"
+#include "Widgets/Input/SCheckBox.h"
 #include "Widgets/Text/STextBlock.h"
 
 #define LOCTEXT_NAMESPACE "GameFeatures"
@@ -174,9 +175,28 @@ void FGameFeatureDataDetailsCustomization::CustomizeDetails(IDetailLayoutBuilder
 				SNew(STextBlock)
 				.Text(this, &FGameFeatureDataDetailsCustomization::GetTagConfigPathText)
 				.Font(DetailBuilder.GetDetailFont())
-			];
-
+		];
 
+		if (FPlatformMisc::IsDebuggerPresent())
+		{
+			const FText Label = LOCTEXT("BreakStateLabel", "Break when State changes");
+			const FText Tooltip = LOCTEXT("BreakStateTooltip", "When enabled, will trigger a breakpoint any time the state of this plugin changes");
+			TopCategory.AddCustomRow(Label)
+			.NameContent()
+			[
+				SNew(STextBlock)
+				.Text(Label)
+				.ToolTipText(Tooltip)
+				.Font(DetailBuilder.GetDetailFont())
+			]
+			.ValueContent()
+			[
+				SNew(SCheckBox)
+				.IsChecked( this, &FGameFeatureDataDetailsCustomization::GetDebugStateEnabled)
+				.OnCheckStateChanged( this, &FGameFeatureDataDetailsCustomization::SetDebugStateEnabled)
+				.ToolTipText(Tooltip)
+			];
+		}
 //@TODO: This disables the mode switcher widget too (and it's a const cast hack...)
 // 		if (IDetailsView* ConstHackDetailsView = const_cast<IDetailsView*>(DetailBuilder.GetDetailsView()))
 // 		{
@@ -226,6 +246,23 @@ EVisibility FGameFeatureDataDetailsCustomization::GetVisbililty() const
 	return (GetCurrentState() == EGameFeaturePluginState::Active) ? EVisibility::Visible : EVisibility::Collapsed;
 }
 
+ECheckBoxState FGameFeatureDataDetailsCustomization::GetDebugStateEnabled() const
+{
+	if (PluginURL.IsEmpty())
+	{
+		return ECheckBoxState::Undetermined;
+	}
+	return UGameFeaturesSubsystem::Get().GetPluginDebugStateEnabled(PluginURL) ? ECheckBoxState::Checked : ECheckBoxState::Unchecked;
+}
+
+void FGameFeatureDataDetailsCustomization::SetDebugStateEnabled(ECheckBoxState InCheckState)
+{
+	if (!PluginURL.IsEmpty())
+	{
+		return UGameFeaturesSubsystem::Get().SetPluginDebugStateEnabled(PluginURL, InCheckState == ECheckBoxState::Checked);
+	}
+}
+
 FText FGameFeatureDataDetailsCustomization::GetInitialStateText() const
 {
 	const EBuiltInAutoState AutoState = UGameFeaturesSubsystem::DetermineBuiltInInitialFeatureState(PluginPtr->GetDescriptor().CachedJson, FString());
diff --git a/Engine/Plugins/Runtime/GameFeatures/Source/GameFeaturesEditor/Private/GameFeatureDataDetailsCustomization.h b/Engine/Plugins/Runtime/GameFeatures/Source/GameFeaturesEditor/Private/GameFeatureDataDetailsCustomization.h
--- a/Engine/Plugins/Runtime/GameFeatures/Source/GameFeaturesEditor/Private/GameFeatureDataDetailsCustomization.h
+++ b/Engine/Plugins/Runtime/GameFeatures/Source/GameFeaturesEditor/Private/GameFeatureDataDetailsCustomization.h
@@ -13,6 +13,8 @@ class SErrorText;
 class IPlugin;
 struct EVisibility;
 
+enum class ECheckBoxState : uint8;
+
 //////////////////////////////////////////////////////////////////////////
 // FGameFeatureDataDetailsCustomization
 
@@ -32,6 +34,9 @@ protected:
 	EGameFeaturePluginState GetCurrentState() const;
 	EVisibility GetVisbililty() const;
 
+	ECheckBoxState GetDebugStateEnabled() const;
+	void SetDebugStateEnabled(ECheckBoxState InCheckState);
+	
 	FText GetInitialStateText() const;
 	FText GetTagConfigPathText() const;
 
