UE: Fix metadata for GFP updatestatemachine in insights traces to handle cases before the plugin has determined its name

diff --git a/Engine/Plugins/Runtime/GameFeatures/Source/GameFeatures/Private/GameFeaturePluginStateMachine.cpp b/Engine/Plugins/Runtime/GameFeatures/Source/GameFeatures/Private/GameFeaturePluginStateMachine.cpp
--- a/Engine/Plugins/Runtime/GameFeatures/Source/GameFeatures/Private/GameFeaturePluginStateMachine.cpp
+++ b/Engine/Plugins/Runtime/GameFeatures/Source/GameFeatures/Private/GameFeaturePluginStateMachine.cpp
@@ -4280,7 +4280,7 @@ void UGameFeaturePluginStateMachine::UpdateStateMachine()
 	}
 
 	UE_TRACE_LOG_SCOPED_T(Cpu, GFP_UpdateStateMachine, CpuChannel)
-		<< GFP_UpdateStateMachine.PluginName(*GetPluginName());
+		<< GFP_UpdateStateMachine.PluginName(*GetGameFeatureName());
 
 	TOptional<TGuardValue<bool>> ScopeGuard(InPlace, bInUpdateStateMachine, true);
 	
