Fix an issue where LoadingBuiltInGameFeaturePlugins was dependent on the current enabled plugins. This can  cause timing issues and inconsistent results as BuiltInGameFeaturePlugins can be called at any time after serveral install bundle GameFeaturePlugins have been loaded. If LoadBuiltInGameFeaturePlugins was called during this time it could make non-built-In plugins have their LoadedAsBuiltIn flag set to true.

Moved and exposed new logic in the plugin manager to get at the list of plugins that were built for the current target exe. Only plugins that were built for the target exe, that means a plugin which was compiled for this exe upfront. Only those plugins should be considered ones that are built in for this call.

Added new CVar:
GameFeaturePlugin.TrimNonStartupEnabledPlugins
This can be used to ignore this new behaviour. All enabled plugins will be treated as built in.

diff --git a/Engine/Plugins/Runtime/GameFeatures/Source/GameFeatures/Private/GameFeaturesSubsystem.cpp b/Engine/Plugins/Runtime/GameFeatures/Source/GameFeatures/Private/GameFeaturesSubsystem.cpp
--- a/Engine/Plugins/Runtime/GameFeatures/Source/GameFeatures/Private/GameFeaturesSubsystem.cpp
+++ b/Engine/Plugins/Runtime/GameFeatures/Source/GameFeatures/Private/GameFeaturesSubsystem.cpp
@@ -102,6 +102,13 @@ namespace UE::GameFeatures
 		ECVF_Default);
 #endif // !UE_BUILD_SHIPPING
 
+	static bool bTrimNonStartupEnabledPlugins = true;
+	static FAutoConsoleVariableRef CVarTrimNonBuiltInPlugins(
+		TEXT("GameFeaturePlugin.TrimNonStartupEnabledPlugins"),
+		bTrimNonStartupEnabledPlugins,
+		TEXT("When call LoadBuiltInGameFeaturePlugins should only the plugins that were built by this target exe be considered built in."),
+		ECVF_Default);
+
 	TOptional<FString> GetPluginUrlForConsoleCommand(const TArray<FString>& Args, FOutputDevice& Ar)
 	{
 		TOptional<FString> PluginURL;
@@ -2253,6 +2260,16 @@ void UGameFeaturesSubsystem::LoadBuiltInGameFeaturePlugins(FBuiltInPluginAdditio
 		LoadContext->Results.Add(Plugin->GetName(), MakeError("Pending"));
 	}
 
+	if (UE::GameFeatures::bTrimNonStartupEnabledPlugins)
+	{
+		const TSet<FString>& CompiledInPlugins = IPluginManager::Get().GetPluginsEnabledAtStartup();
+		EnabledPlugins.RemoveAllSwap([&CompiledInPlugins](const TSharedRef<IPlugin>& Plugin)
+			{
+				const bool bRemove = !CompiledInPlugins.Contains(Plugin->GetName());
+				return bRemove;
+			});
+	}
+
 	const int32 NumPluginsToLoad = EnabledPlugins.Num();
 	UE_LOG(LogGameFeatures, Log, TEXT("Loading %i builtins"), NumPluginsToLoad);
 
@@ -2277,6 +2294,7 @@ void UGameFeaturesSubsystem::LoadBuiltInGameFeaturePlugins(FBuiltInPluginAdditio
 		}
 		return Dependencies;
 	};
+
 	Algo::TopologicalSort(EnabledPlugins, GetPluginDependencies);
 
 	for (const TSharedRef<IPlugin>& Plugin : EnabledPlugins)
diff --git a/Engine/Source/Runtime/Projects/Private/PluginManager.cpp b/Engine/Source/Runtime/Projects/Private/PluginManager.cpp
--- a/Engine/Source/Runtime/Projects/Private/PluginManager.cpp
+++ b/Engine/Source/Runtime/Projects/Private/PluginManager.cpp
@@ -1837,6 +1837,8 @@ bool FPluginManager::ConfigureEnabledPlugins()
 				PluginsEnabledForTarget.Add(Pair.Key);
 			}
 		}
+		// Initially this will be correct. PluginsEnabledForTarget will be updated with new mounts and this will stop including compiled plugins.
+		PluginsEnabledAtStartup.Append(PluginsEnabledForTarget);
 
 		// If we made it here, we have all the required plugins
 		bHaveAllRequiredPlugins = true;
@@ -3117,6 +3119,12 @@ TArray<TSharedRef<IPlugin>> FPluginManager::GetDiscoveredPlugins()
 	return Plugins;
 }
 
+const TSet<FString>& FPluginManager::GetPluginsEnabledAtStartup() const
+{
+	ensureMsgf(!PluginsEnabledAtStartup.IsEmpty(), TEXT("ConfigureEnabledPlugins hasn't been called yet."));
+	return PluginsEnabledAtStartup;
+}
+
 #if WITH_EDITOR
 const TSet<FString>& FPluginManager::GetBuiltInPluginNames() const
 {
diff --git a/Engine/Source/Runtime/Projects/Private/PluginManager.h b/Engine/Source/Runtime/Projects/Private/PluginManager.h
--- a/Engine/Source/Runtime/Projects/Private/PluginManager.h
+++ b/Engine/Source/Runtime/Projects/Private/PluginManager.h
@@ -200,6 +200,7 @@ public:
 	virtual TArray<TSharedRef<IPlugin>> GetEnabledPluginsWithContentOrVerse() const override;
 	virtual TArray<TSharedRef<IPlugin>> GetDiscoveredPlugins() override;
 
+	virtual const TSet<FString>& GetPluginsEnabledAtStartup() const override;
 #if WITH_EDITOR
 	virtual const TSet<FString>& GetBuiltInPluginNames() const override;
 	virtual TSharedPtr<IPlugin> GetModuleOwnerPlugin(FName ModuleName) const override;
@@ -334,6 +335,9 @@ private:
 	/** Set of plugins enabled for the current target. Plugins in list might not be currently enabled. */
 	TSet<FString> PluginsEnabledForTarget;
 
+	/** Set of plugins which were compiled for the target and enabled at startup. */
+	TSet<FString> PluginsEnabledAtStartup;
+
 #if WITH_EDITOR
 	/** Names of built-in plugins */
 	TSet<FString> BuiltInPluginNames;
diff --git a/Engine/Source/Runtime/Projects/Public/Interfaces/IPluginManager.h b/Engine/Source/Runtime/Projects/Public/Interfaces/IPluginManager.h
--- a/Engine/Source/Runtime/Projects/Public/Interfaces/IPluginManager.h
+++ b/Engine/Source/Runtime/Projects/Public/Interfaces/IPluginManager.h
@@ -434,6 +434,15 @@ public:
 	 */
 	virtual TArray<TSharedRef<IPlugin>> GetDiscoveredPlugins() = 0;
 
+	/**
+	 * Returns the set of plugins compiled into the current application and enabled at startup.
+	 * For internal use only.
+	 * TODO: Needs improved for the GameFeaturePlugin system to one day remove the concept of LoadedAsBuiltIn
+	 * 
+	 * @return set of plugin names.
+	 */
+	UE_INTERNAL virtual const TSet<FString>& GetPluginsEnabledAtStartup() const = 0;
+
 #if WITH_EDITOR
 	/**
 	 * Returns the set of built-in plugin names
