Add an overload for TerminateGameFeaturePlugin that can handle a plugin list

diff --git a/Engine/Plugins/Runtime/GameFeatures/Source/GameFeatures/Private/GameFeaturesSubsystem.cpp b/Engine/Plugins/Runtime/GameFeatures/Source/GameFeatures/Private/GameFeaturesSubsystem.cpp
--- a/Engine/Plugins/Runtime/GameFeatures/Source/GameFeatures/Private/GameFeaturesSubsystem.cpp
+++ b/Engine/Plugins/Runtime/GameFeatures/Source/GameFeatures/Private/GameFeaturesSubsystem.cpp
@@ -1937,6 +1937,66 @@ void UGameFeaturesSubsystem::TerminateGameFeaturePlugin(const FString& PluginURL
 	}
 }
 
+void UGameFeaturesSubsystem::TerminateGameFeaturePlugin(TConstArrayView<FString> PluginURLs, const FMultipleGameFeaturePluginsTerminated& CompleteDelegate)
+{
+	struct FContext
+	{
+		TMap<FString, UE::GameFeatures::FResult> Results;
+		FMultipleGameFeaturePluginsTerminated CompleteDelegate;
+
+		int32 NumPluginsTerminated = 0;
+		bool bPushedTagsBroadcast = false;
+
+		FContext()
+		{
+			if (!IsEngineExitRequested())
+			{
+				UGameplayTagsManager::Get().PushDeferOnGameplayTagTreeChangedBroadcast();
+				bPushedTagsBroadcast = true;
+			}
+			else if (UGameplayTagsManager* TagsManager = UGameplayTagsManager::GetIfAllocated())
+			{
+				TagsManager->PushDeferOnGameplayTagTreeChangedBroadcast();
+				bPushedTagsBroadcast = true;
+			}
+		}
+
+		~FContext()
+		{
+			if (bPushedTagsBroadcast)
+			{
+				if (UGameplayTagsManager* TagsManager = UGameplayTagsManager::GetIfAllocated())
+				{
+					TagsManager->PopDeferOnGameplayTagTreeChangedBroadcast();
+				}
+			}
+
+			CompleteDelegate.ExecuteIfBound(Results);
+		}
+	};
+	TSharedRef<FContext> Context = MakeShared<FContext>();
+	Context->CompleteDelegate = CompleteDelegate;
+
+	Context->Results.Reserve(PluginURLs.Num());
+	for (const FString& PluginURL : PluginURLs)
+	{
+		Context->Results.Add(PluginURL, MakeError("Pending"));
+	}
+
+	const int32 NumPlugins = PluginURLs.Num();
+	UE_LOG(LogGameFeatures, Log, TEXT("Terminating %i GFPs"), NumPlugins);
+
+	for (const FString& PluginURL : PluginURLs)
+	{
+		TerminateGameFeaturePlugin(PluginURL, FGameFeaturePluginTerminateComplete::CreateLambda([Context, PluginURL](const UE::GameFeatures::FResult& Result)
+		{
+			Context->Results.Add(PluginURL, Result);
+			++Context->NumPluginsTerminated;
+			UE_LOG(LogGameFeatures, VeryVerbose, TEXT("Finished Terminating %i GFPs"), Context->NumPluginsTerminated);
+		}));
+	}
+}
+
 void UGameFeaturesSubsystem::CancelGameFeatureStateChange(const FString& PluginURL)
 {
 	CancelGameFeatureStateChange(PluginURL, FGameFeaturePluginChangeStateComplete());
diff --git a/Engine/Plugins/Runtime/GameFeatures/Source/GameFeatures/Public/GameFeaturesSubsystem.h b/Engine/Plugins/Runtime/GameFeatures/Source/GameFeatures/Public/GameFeaturesSubsystem.h
--- a/Engine/Plugins/Runtime/GameFeatures/Source/GameFeatures/Public/GameFeaturesSubsystem.h
+++ b/Engine/Plugins/Runtime/GameFeatures/Source/GameFeatures/Public/GameFeaturesSubsystem.h
@@ -215,6 +215,7 @@ using FMultipleGameFeaturePluginChangeStateComplete = TDelegate<void(const TMap<
 
 using FBuiltInGameFeaturePluginsLoaded = FMultipleGameFeaturePluginChangeStateComplete;
 using FMultipleGameFeaturePluginsLoaded = FMultipleGameFeaturePluginChangeStateComplete;
+using FMultipleGameFeaturePluginsTerminated = FMultipleGameFeaturePluginChangeStateComplete;
 
 enum class EBuiltInAutoState : uint8
 {
@@ -596,6 +597,7 @@ public:
 	/** Terminate the GameFeaturePlugin and remove all associated plugin tracking data. */
 	void TerminateGameFeaturePlugin(const FString& PluginURL);
 	void TerminateGameFeaturePlugin(const FString& PluginURL, const FGameFeaturePluginTerminateComplete& CompleteDelegate);
+	void TerminateGameFeaturePlugin(TConstArrayView<FString> PluginURLs, const FMultipleGameFeaturePluginsTerminated& CompleteDelegate);
 	
 	/** Attempt to cancel any state change. Calls back when cancelation is complete. Any other pending callbacks will be called with a canceled error. */
 	void CancelGameFeatureStateChange(const FString& PluginURL);
