A start to some basic GameFeaturePlugin tests, to start just a file protocol GFP to test some of the public APIs against
- Can transition to all the public transition points
- Can query information about its PluginProtocol
- Can query its GameFeatureData, depending on what State its in

diff --git a/Engine/Plugins/Runtime/GameFeatures/Source/GameFeatures/Private/Tests/GameFeaturePluginTests.cpp b/Engine/Plugins/Runtime/GameFeatures/Source/GameFeatures/Private/Tests/GameFeaturePluginTests.cpp
new file mode 100644
--- /dev/null
+++ b/Engine/Plugins/Runtime/GameFeatures/Source/GameFeatures/Private/Tests/GameFeaturePluginTests.cpp
@@ -0,0 +1,280 @@
+#include "Misc/AutomationTest.h"
+#include "Misc/ScopeExit.h"
+#include "Tests/AutomationCommon.h"
+#include "Tests/AutomationEditorCommon.h"
+
+#include "GameFeaturePluginOperationResult.h"
+#include "GameFeaturePluginStateMachine.h"
+#include "GameFeaturesSubsystem.h"
+#include "GameFeatureTypes.h"
+
+#if WITH_DEV_AUTOMATION_TESTS
+
+DEFINE_LATENT_AUTOMATION_COMMAND_ONE_PARAMETER(FWaitForTrue, bool*, bVariableToWaitFor);
+bool FWaitForTrue::Update()
+{
+    return *bVariableToWaitFor;
+}
+
+DEFINE_LATENT_AUTOMATION_COMMAND_ONE_PARAMETER(FExecuteFunction, TFunction<bool()>, Function);
+bool FExecuteFunction::Update()
+{
+	return Function();
+}
+
+EGameFeaturePluginState ConvertTargetStateToPluginState(const EGameFeatureTargetState TargetState)
+{
+	switch (TargetState)
+	{
+		case EGameFeatureTargetState::Installed:
+			return EGameFeaturePluginState::Installed;
+		case EGameFeatureTargetState::Registered:
+			return EGameFeaturePluginState::Registered;
+		case EGameFeatureTargetState::Loaded:
+			return EGameFeaturePluginState::Loaded;
+		case EGameFeatureTargetState::Active:
+			return EGameFeaturePluginState::Active;
+		default:
+			break;
+	}
+
+	return EGameFeaturePluginState::MAX;
+}
+
+class FTestGameFeaturePluginBase : public FAutomationTestBase
+{
+public:
+	FTestGameFeaturePluginBase(const FString& InName, const bool bInComplexTask)
+		: FAutomationTestBase(InName, bInComplexTask)
+	{
+	}
+
+	~FTestGameFeaturePluginBase()
+	{
+	}
+
+	bool IsPluginInPluginStateRange(const FGameFeaturePluginStateRange PluginStateRange)
+	{
+		EGameFeaturePluginState CurrentPluginState = UGameFeaturesSubsystem::Get().GetPluginState(GFPFileURL);
+
+		return PluginStateRange.Contains(CurrentPluginState);
+	}
+
+	void LatentTestPluginState(const FGameFeaturePluginStateRange PluginStateRange)
+	{
+		ADD_LATENT_AUTOMATION_COMMAND(FExecuteFunction([this, PluginStateRange]
+		{
+			TestTrue(FString::Printf(TEXT("Plugin in %s state, expected plugin state in range (%s, %s)"),
+				*UE::GameFeatures::ToString(UGameFeaturesSubsystem::Get().GetPluginState(GFPFileURL)), *UE::GameFeatures::ToString(PluginStateRange.MinState), *UE::GameFeatures::ToString(PluginStateRange.MaxState)),
+				IsPluginInPluginStateRange(PluginStateRange));
+
+			return true;
+		}));
+	}
+
+	void LatentTestTransitionGFP(const EGameFeatureTargetState TargetState)
+	{
+		ADD_LATENT_AUTOMATION_COMMAND(FExecuteFunction([this, TargetState]
+		{
+			*bAsyncCommandComplete = false;
+
+			UGameFeaturesSubsystem::Get().ChangeGameFeatureTargetState(GFPFileURL, TargetState,
+				FGameFeaturePluginChangeStateComplete::CreateLambda([this, TargetState](const UE::GameFeatures::FResult& Result)
+				{
+					*bAsyncCommandComplete = true;
+					TestFalse(FString::Printf(TEXT("Failed to transition to %s: error: %s"), *LexToString(TargetState), *UE::GameFeatures::ToString(Result)),
+						Result.HasError());
+				}
+			));
+
+			return true;
+		}));
+
+		ADD_LATENT_AUTOMATION_COMMAND(FWaitForTrue(&*bAsyncCommandComplete));
+
+		LatentTestPluginState(FGameFeaturePluginStateRange(ConvertTargetStateToPluginState(TargetState)));
+	}
+
+	void LatentCheckInitialPluginState()
+	{
+		// Check we are somewhere between uninited, and uninstalled for the first time we check this and after we restore the plugin state
+		// depending on the initial state as well as deactivating/terminating the plugin we should be in the Terminal or UnknownStatus node
+		LatentTestPluginState(FGameFeaturePluginStateRange(EGameFeaturePluginState::Uninitialized, EGameFeaturePluginState::Uninstalled));
+	}
+
+	void LatentRestorePluginState()
+	{
+		ADD_LATENT_AUTOMATION_COMMAND(FExecuteFunction([this]
+		{
+			*bAsyncCommandComplete = false;
+
+			// We are in an uninstalled/terminal/not setup state. Dont try to Deactivate/Terminate when we are not Activated/Installed
+			if (IsPluginInPluginStateRange(FGameFeaturePluginStateRange(EGameFeaturePluginState::Uninitialized, EGameFeaturePluginState::Uninstalled)))
+			{
+				*bAsyncCommandComplete = true;
+				return true;
+			}
+
+			UGameFeaturesSubsystem::Get().DeactivateGameFeaturePlugin(GFPFileURL,
+				FGameFeaturePluginReleaseComplete::CreateLambda([this](const UE::GameFeatures::FResult& Result)
+				{
+					*bAsyncCommandComplete = true;
+					TestFalse(FString::Printf(TEXT("Failed to deactivate plugin, error: %s"), *UE::GameFeatures::ToString(Result)),
+						Result.HasError());
+				}
+			));
+
+			return true;
+		}));
+		ADD_LATENT_AUTOMATION_COMMAND(FWaitForTrue(&*bAsyncCommandComplete));
+
+		ADD_LATENT_AUTOMATION_COMMAND(FExecuteFunction([this]
+		{
+			*bAsyncCommandComplete = false;
+
+			// We are in an uninstalled/terminal/not setup state. Dont try to Deactivate/Terminate when we are not Activated/Installed
+			if (IsPluginInPluginStateRange(FGameFeaturePluginStateRange(EGameFeaturePluginState::Uninitialized, EGameFeaturePluginState::Uninstalled)))
+			{
+				*bAsyncCommandComplete = true;
+				return true;
+			}
+
+			UGameFeaturesSubsystem::Get().TerminateGameFeaturePlugin(GFPFileURL,
+				FGameFeaturePluginReleaseComplete::CreateLambda([this](const UE::GameFeatures::FResult& Result)
+				{
+					*bAsyncCommandComplete = true;
+					TestFalse(FString::Printf(TEXT("Failed to terminate plugin, error: %s"), *UE::GameFeatures::ToString(Result)),
+						Result.HasError());
+				}
+			));
+
+			return true;
+		}));
+		ADD_LATENT_AUTOMATION_COMMAND(FWaitForTrue(&*bAsyncCommandComplete));
+	}
+
+	// For now hard-coded into EngineTest area but can always be adjusted later
+	FString GFPPluginPath = TEXT("../../../EngineTest/Plugins/GameFeatures/GameFeatureEngineTestC/GameFeatureEngineTestC.uplugin");
+	FString GFPFileURL = FString(TEXT("file:")) + GFPPluginPath;
+	TSharedRef<bool> bAsyncCommandComplete = MakeShared<bool>(false);
+};
+
+IMPLEMENT_CUSTOM_SIMPLE_AUTOMATION_TEST(FGameFeatureSubsystemTestChangeState, FTestGameFeaturePluginBase, "GameFeaturePlugin.Subsystem.ChangeTargetState", EAutomationTestFlags::EditorContext | EAutomationTestFlags::EngineFilter);
+bool FGameFeatureSubsystemTestChangeState::RunTest(const FString& Parameters)
+{
+	// Ensure if the test was canceled we restore the plugin back to an deactivated/terminated state
+	LatentRestorePluginState();
+	LatentCheckInitialPluginState();
+	ON_SCOPE_EXIT
+	{
+		LatentRestorePluginState();
+	};
+
+	LatentTestTransitionGFP(EGameFeatureTargetState::Installed);
+	LatentTestTransitionGFP(EGameFeatureTargetState::Registered);
+	LatentTestTransitionGFP(EGameFeatureTargetState::Loaded);
+	LatentTestTransitionGFP(EGameFeatureTargetState::Active);
+
+    return true;
+}
+
+IMPLEMENT_CUSTOM_SIMPLE_AUTOMATION_TEST(FGameFeatureSubsystemTestUninstall, FTestGameFeaturePluginBase, "GameFeaturePlugin.Subsystem.FilePluginProtocol", EAutomationTestFlags::EditorContext | EAutomationTestFlags::EngineFilter);
+bool FGameFeatureSubsystemTestUninstall::RunTest(const FString& Parameters)
+{
+	// Ensure if the test was canceled we restore the plugin back to an deactivated/terminated state
+	LatentRestorePluginState();
+	LatentCheckInitialPluginState();
+	ON_SCOPE_EXIT
+	{
+		LatentRestorePluginState();
+	};
+
+	// Get us into an installed state so we can query info about the GFP
+	LatentTestTransitionGFP(EGameFeatureTargetState::Installed);
+
+	ADD_LATENT_AUTOMATION_COMMAND(FExecuteFunction([this]
+	{
+		EGameFeaturePluginProtocol FilePluginProtocol = UGameFeaturesSubsystem::Get().GetPluginURLProtocol(GFPFileURL);
+		if (!TestEqual(FString::Printf(TEXT("Expected PluginProtocol to be File but was %i"), (int)FilePluginProtocol),
+				FilePluginProtocol, EGameFeaturePluginProtocol::File))
+		{
+			return true;
+		}
+
+		if (!TestTrue(TEXT("Expected PluginProtocol to be File but was not"),
+				UGameFeaturesSubsystem::Get().IsPluginURLProtocol(GFPFileURL, EGameFeaturePluginProtocol::File)))
+		{
+			return true;
+		}
+
+		EGameFeaturePluginProtocol PluginProtocol;
+		FStringView PluginPath;
+		if (!TestTrue(TEXT("Failed to parse plugin URL"),
+				UGameFeaturesSubsystem::Get().ParsePluginURL(GFPFileURL, &PluginProtocol, &PluginPath)))
+		{
+			return true;
+		}
+
+		if (!TestEqual(FString::Printf(TEXT("Expected PluginProtocol to be File but was %i"), (int)PluginProtocol),
+				PluginProtocol, EGameFeaturePluginProtocol::File))
+		{
+			return true;
+		}
+
+		if (!TestEqual(FString::Printf(TEXT("Expected parsed PluginPath %s to equal %s"), PluginPath.GetData(), *GFPPluginPath),
+			PluginPath, GFPPluginPath))
+		{
+			return true;
+		}
+
+		return true;
+	}));
+
+	return true;
+}
+
+IMPLEMENT_CUSTOM_SIMPLE_AUTOMATION_TEST(FGameFeatureSubsystemTestGetGameFeatureData, FTestGameFeaturePluginBase, "GameFeaturePlugin.Subsystem.GetGameFeatureData", EAutomationTestFlags::EditorContext | EAutomationTestFlags::EngineFilter);
+bool FGameFeatureSubsystemTestGetGameFeatureData::RunTest(const FString& Parameters)
+{
+	// Ensure if the test was canceled we restore the plugin back to an deactivated/terminated state
+	LatentRestorePluginState();
+	LatentCheckInitialPluginState();
+	ON_SCOPE_EXIT
+	{
+		LatentRestorePluginState();
+	};
+
+	LatentTestTransitionGFP(EGameFeatureTargetState::Installed);
+
+	ADD_LATENT_AUTOMATION_COMMAND(FExecuteFunction([this]
+	{
+		const UGameFeatureData* GameFeatureData = UGameFeaturesSubsystem::Get().GetGameFeatureDataForRegisteredPluginByURL(GFPFileURL);
+		TestNull(TEXT("GameFeatureData is not NULL, GFP is only in the Installed state and should not have any GameFeatureData"), GameFeatureData);
+
+		return true;
+	}));
+
+	LatentTestTransitionGFP(EGameFeatureTargetState::Registered);
+
+	ADD_LATENT_AUTOMATION_COMMAND(FExecuteFunction([this]
+	{
+		const UGameFeatureData* GameFeatureData = UGameFeaturesSubsystem::Get().GetGameFeatureDataForRegisteredPluginByURL(GFPFileURL);
+		TestNotNull(TEXT("GameFeatureData is NULL, but the GFP should have a valid GameFeatureData"), GameFeatureData);
+
+		return true;
+	}));
+
+	LatentTestTransitionGFP(EGameFeatureTargetState::Active);
+
+	ADD_LATENT_AUTOMATION_COMMAND(FExecuteFunction([this]
+	{
+		const UGameFeatureData* GameFeatureData = UGameFeaturesSubsystem::Get().GetGameFeatureDataForActivePluginByURL(GFPFileURL);
+		TestNotNull(TEXT("GameFeatureData is NULL, but the GFP should have a valid GameFeatureData"), GameFeatureData);
+
+		return true;
+	}));
+
+	return true;
+}
+
+#endif // WITH_DEV_AUTOMATION_TESTS
