Fix test only for Editor trying to compile on Clients

diff --git a/Engine/Plugins/Runtime/GameFeatures/Source/GameFeatures/Private/Tests/GameFeaturePluginTests.cpp b/Engine/Plugins/Runtime/GameFeatures/Source/GameFeatures/Private/Tests/GameFeaturePluginTests.cpp
--- a/Engine/Plugins/Runtime/GameFeatures/Source/GameFeatures/Private/Tests/GameFeaturePluginTests.cpp
+++ b/Engine/Plugins/Runtime/GameFeatures/Source/GameFeatures/Private/Tests/GameFeaturePluginTests.cpp
@@ -1,3 +1,5 @@
+#if WITH_DEV_AUTOMATION_TESTS && WITH_EDITOR
+
 #include "Misc/AutomationTest.h"
 #include "Misc/ScopeExit.h"
 #include "Tests/AutomationCommon.h"
@@ -8,8 +10,6 @@
 #include "GameFeaturesSubsystem.h"
 #include "GameFeatureTypes.h"
 
-#if WITH_DEV_AUTOMATION_TESTS
-
 DEFINE_LATENT_AUTOMATION_COMMAND_ONE_PARAMETER(FWaitForTrue, bool*, bVariableToWaitFor);
 bool FWaitForTrue::Update()
 {
@@ -277,4 +277,4 @@ bool FGameFeatureSubsystemTestGetGameFeatureData::RunTest(const FString& Paramet
 	return true;
 }
 
-#endif // WITH_DEV_AUTOMATION_TESTS
+#endif // WITH_DEV_AUTOMATION_TESTS && WITH_EDITOR
