Write a test to cover an issue where some plugins would be leaked as a builtin plugin
where it was actually loaded externally and not truely a compiled in plugin

diff --git a/Engine/Plugins/Runtime/GameFeatures/Source/GameFeatures/Private/Tests/GameFeaturePluginTests.cpp b/Engine/Plugins/Runtime/GameFeatures/Source/GameFeatures/Private/Tests/GameFeaturePluginTests.cpp
--- a/Engine/Plugins/Runtime/GameFeatures/Source/GameFeatures/Private/Tests/GameFeaturePluginTests.cpp
+++ b/Engine/Plugins/Runtime/GameFeatures/Source/GameFeatures/Private/Tests/GameFeaturePluginTests.cpp
@@ -279,4 +279,63 @@ bool FGameFeatureSubsystemTestGetGameFeatureData::RunTest(const FString& Paramet
 	return true;
 }
 
+// This test is testing that non-compiled in plugins do not get marked as built in once they are loaded through external APIs
+// To see the test fail set GameFeaturePlugin.TrimNonStartupEnabledPlugins=false, which will go back to the old way the plugin system would handle new plugins not set as built in
+IMPLEMENT_CUSTOM_SIMPLE_AUTOMATION_TEST(FGameFeatureSubsystemTestNonBuiltinPluginDoesntConvertToBuiltinPlugin, FTestGameFeaturePluginBase, "GameFeaturePlugin.Subsystem.NonBuiltinPluginDoesntConvertToBuiltinPlugin", EAutomationTestFlags::EditorContext | EAutomationTestFlags::EngineFilter);
+bool FGameFeatureSubsystemTestNonBuiltinPluginDoesntConvertToBuiltinPlugin::RunTest(const FString& Parameters)
+{
+	// Ensure if the test was canceled we restore the plugin back to an deactivated/terminated state
+	LatentRestorePluginState();
+	LatentCheckInitialPluginState();
+	ON_SCOPE_EXIT
+	{
+		LatentRestorePluginState();
+	};
+
+	LatentTestTransitionGFP(EGameFeatureTargetState::Installed);
+	LatentTestTransitionGFP(EGameFeatureTargetState::Registered);
+
+	// Test we get to registered, installed -> mounted which will get our plugin in the enabled/mounted state
+	ADD_LATENT_AUTOMATION_COMMAND(FExecuteFunction([this]
+	{
+		TestFalse(FString::Printf(TEXT("WasGameFeaturePluginLoadedAsBuiltIn on GFP %s to be false but was true"), *GFPFileURL),
+			UGameFeaturesSubsystem::Get().WasGameFeaturePluginLoadedAsBuiltIn(GFPFileURL));
+		return true;
+	}));
+
+	ADD_LATENT_AUTOMATION_COMMAND(FExecuteFunction([this]
+	{
+		*bAsyncCommandComplete = false;
+
+		auto AdditionalFilter = [&](const FString& PluginFilename, const FGameFeaturePluginDetails& PluginDetails, FBuiltInGameFeaturePluginBehaviorOptions& OutOptions) -> bool
+		{
+			return true;
+		};
+
+		UGameFeaturesSubsystem::Get().LoadBuiltInGameFeaturePlugins(AdditionalFilter,
+			FBuiltInGameFeaturePluginsLoaded::CreateLambda([this](const TMap<FString, UE::GameFeatures::FResult>& Results)
+			{
+				*bAsyncCommandComplete = true;
+				for (const TPair<FString, UE::GameFeatures::FResult>& Result : Results)
+				{
+					TestFalse(FString::Printf(TEXT("Failed to LoadBuiltInGameFeaturePlugins on %s error: %s"), *Result.Get<0>(), *UE::GameFeatures::ToString(Result.Get<1>())),
+						Result.Get<1>().HasError());
+				}
+			}
+		));
+
+		return true;
+	}));
+	ADD_LATENT_AUTOMATION_COMMAND(FWaitForTrue(&*bAsyncCommandComplete));
+
+	ADD_LATENT_AUTOMATION_COMMAND(FExecuteFunction([this]
+	{
+		TestFalse(FString::Printf(TEXT("WasGameFeaturePluginLoadedAsBuiltIn on GFP %s to be false but was true"), *GFPFileURL),
+			UGameFeaturesSubsystem::Get().WasGameFeaturePluginLoadedAsBuiltIn(GFPFileURL));
+		return true;
+	}));
+
+	return true;
+}
+
 #endif // WITH_DEV_AUTOMATION_TESTS && WITH_EDITOR
