Restore non-async loading of GFP localization data in commandlets

Fixes a case where CL 41606150 caused some GFPs never to finish activating.

#rnx

diff --git a/Engine/Plugins/Runtime/GameFeatures/Source/GameFeatures/Private/GameFeaturePluginStateMachine.cpp b/Engine/Plugins/Runtime/GameFeatures/Source/GameFeatures/Private/GameFeaturePluginStateMachine.cpp
--- a/Engine/Plugins/Runtime/GameFeatures/Source/GameFeatures/Private/GameFeaturePluginStateMachine.cpp
+++ b/Engine/Plugins/Runtime/GameFeatures/Source/GameFeatures/Private/GameFeaturePluginStateMachine.cpp
@@ -126,25 +126,25 @@ namespace UE::GameFeatures
 		StateProperties.bIsLoadingLocalizationData = true;
 		IPluginManager::Get().MountExplicitlyLoadedPluginLocalizationData(StateProperties.PluginName, [WeakMachine = TWeakObjectPtr<UGameFeaturePluginStateMachine>(CurrentMachine), &StateProperties](bool bLoadedLocalization, const FString& PluginName)
 		{
-			if (IsInGameThread())
+			if (UGameFeaturePluginStateMachine* StateMachine = WeakMachine.Get();
+				StateMachine && ensureAlways(&StateMachine->GetProperties() == &StateProperties))
 			{
-				if (UGameFeaturePluginStateMachine* StateMachine = WeakMachine.Get();
-					StateMachine && ensureAlways(&StateMachine->GetProperties() == &StateProperties))
+				if (IsInGameThread())
 				{
 					StateProperties.bIsLoadingLocalizationData = false;
 				}
-			}
-			else
-			{
-				ExecuteOnGameThread(UE_SOURCE_LOCATION, [WeakMachine, &StateProperties]()
+				else if (StateProperties.AllowAsyncLoading())
 				{
-					if (UGameFeaturePluginStateMachine* StateMachine = WeakMachine.Get();
-						StateMachine && ensureAlways(&StateMachine->GetProperties() == &StateProperties))
+					ExecuteOnGameThread(UE_SOURCE_LOCATION, [WeakMachine, &StateProperties]()
 					{
-						StateProperties.bIsLoadingLocalizationData = false;
-						StateProperties.OnRequestUpdateStateMachine.ExecuteIfBound();
-					}
-				});
+						if (UGameFeaturePluginStateMachine* StateMachine = WeakMachine.Get();
+							StateMachine && ensureAlways(&StateMachine->GetProperties() == &StateProperties))
+						{
+							StateProperties.bIsLoadingLocalizationData = false;
+							StateProperties.OnRequestUpdateStateMachine.ExecuteIfBound();
+						}
+					});
+				}
 			}
 		});
 	}
@@ -3772,7 +3772,13 @@ struct FGameFeaturePluginState_Activating : public FGameFeaturePluginState
 		// If this plugin caused localization data to load, we need that load to finish before marking it as active
 		if (StateProperties.bIsLoadingLocalizationData)
 		{
-			return;
+			if (AllowAsyncLoading())
+			{
+				return;
+			}
+
+			FTextLocalizationManager::Get().WaitForAsyncTasks();
+			StateProperties.bIsLoadingLocalizationData = false;
 		}
 
 		if (IsWaitingForBatchProcessing())
