Don't call AllowAsyncLoading on the worker thread

#rnx

diff --git a/Engine/Plugins/Runtime/GameFeatures/Source/GameFeatures/Private/GameFeaturePluginStateMachine.cpp b/Engine/Plugins/Runtime/GameFeatures/Source/GameFeatures/Private/GameFeaturePluginStateMachine.cpp
--- a/Engine/Plugins/Runtime/GameFeatures/Source/GameFeatures/Private/GameFeaturePluginStateMachine.cpp
+++ b/Engine/Plugins/Runtime/GameFeatures/Source/GameFeatures/Private/GameFeaturePluginStateMachine.cpp
@@ -124,7 +124,7 @@ namespace UE::GameFeatures
 		check(CurrentMachine && &CurrentMachine->GetProperties() == &StateProperties);
 
 		StateProperties.bIsLoadingLocalizationData = true;
-		IPluginManager::Get().MountExplicitlyLoadedPluginLocalizationData(StateProperties.PluginName, [WeakMachine = TWeakObjectPtr<UGameFeaturePluginStateMachine>(CurrentMachine), &StateProperties](bool bLoadedLocalization, const FString& PluginName)
+		IPluginManager::Get().MountExplicitlyLoadedPluginLocalizationData(StateProperties.PluginName, [WeakMachine = TWeakObjectPtr<UGameFeaturePluginStateMachine>(CurrentMachine), &StateProperties, bAllowAsyncLoading = StateProperties.AllowAsyncLoading()](bool bLoadedLocalization, const FString& PluginName)
 		{
 			if (UGameFeaturePluginStateMachine* StateMachine = WeakMachine.Get();
 				StateMachine && ensureAlways(&StateMachine->GetProperties() == &StateProperties))
@@ -133,7 +133,7 @@ namespace UE::GameFeatures
 				{
 					StateProperties.bIsLoadingLocalizationData = false;
 				}
-				else if (StateProperties.AllowAsyncLoading())
+				else if (bAllowAsyncLoading)
 				{
 					ExecuteOnGameThread(UE_SOURCE_LOCATION, [WeakMachine, &StateProperties]()
 					{
