UE: Add GameFeaturesProjectPolicies option to skip reading plugin details for certain plugins for performance reasons. Useful if you want to only read certain plugins at certain times

diff --git a/Engine/Plugins/Runtime/GameFeatures/Source/GameFeatures/Private/GameFeaturesSubsystem.cpp b/Engine/Plugins/Runtime/GameFeatures/Source/GameFeatures/Private/GameFeaturesSubsystem.cpp
--- a/Engine/Plugins/Runtime/GameFeatures/Source/GameFeatures/Private/GameFeaturesSubsystem.cpp
+++ b/Engine/Plugins/Runtime/GameFeatures/Source/GameFeatures/Private/GameFeaturesSubsystem.cpp
@@ -2516,6 +2516,12 @@ bool UGameFeaturesSubsystem::GetGameFeaturePluginDetailsInternal(const FString&
 {
 	TRACE_CPUPROFILER_EVENT_SCOPE(GFP_GetPluginDetails);
 
+	check(GameSpecificPolicies);
+	if (!GameSpecificPolicies->ShouldReadPluginDetails(PluginDescriptorFilename))
+	{
+		return false;
+	}
+
 	// GFPs are implemented with a plugin so FPluginReferenceDescriptor doesn't know anything about them.
 	// Need a better way of storing GFP specific plugin data...
 
diff --git a/Engine/Plugins/Runtime/GameFeatures/Source/GameFeatures/Public/GameFeaturesProjectPolicies.h b/Engine/Plugins/Runtime/GameFeatures/Source/GameFeatures/Public/GameFeaturesProjectPolicies.h
--- a/Engine/Plugins/Runtime/GameFeatures/Source/GameFeatures/Public/GameFeaturesProjectPolicies.h
+++ b/Engine/Plugins/Runtime/GameFeatures/Source/GameFeatures/Public/GameFeaturesProjectPolicies.h
@@ -72,6 +72,9 @@ public:
 		PRAGMA_ENABLE_DEPRECATION_WARNINGS
 	}
 
+	// Return true if a uplugin's details should be read and false if it should be skipped. Skipped plugins will not be processed as GFPs and skipped as though they didn't exist. Useful to limit the number of uplugin files opened for perf reasons
+	virtual bool ShouldReadPluginDetails(const FString& PluginDescriptorFilename) const { return true; }
+
 	// Called to resolve plugin dependencies, will successfully return an empty string if a dependency is not a GFP.
 	// This may be called with file protocol for built-in plugins in some cases, even if a different protocol is used at runtime.
 	// returns The dependency URL or an error if the dependency could not be resolved
