UE: Add GameFeaturesSubsystem option to LoadBuiltInGameFeaturePlugins amortized over multiple frames with a given rate of plugins per frame. The AdditonalFilter is copied and used in future frames so it cannot be a TFunctionRef and it cannot capture stack variables by reference if it is a lambda

diff --git a/Engine/Plugins/Runtime/GameFeatures/Source/GameFeatures/Private/GameFeaturesSubsystem.cpp b/Engine/Plugins/Runtime/GameFeatures/Source/GameFeatures/Private/GameFeaturesSubsystem.cpp
--- a/Engine/Plugins/Runtime/GameFeatures/Source/GameFeatures/Private/GameFeaturesSubsystem.cpp
+++ b/Engine/Plugins/Runtime/GameFeatures/Source/GameFeatures/Private/GameFeaturesSubsystem.cpp
@@ -2250,6 +2250,17 @@ private:
 #endif // !UE_BUILD_SHIPPING
 
 void UGameFeaturesSubsystem::LoadBuiltInGameFeaturePlugins(FBuiltInPluginAdditionalFilters AdditionalFilter, const FBuiltInGameFeaturePluginsLoaded& InCompleteDelegate /*= FBuiltInGameFeaturePluginsLoaded()*/)
+{
+	FBuiltInPluginAdditionalFilters_Copyable NullCopyable;
+	LoadBuiltInGameFeaturePluginsInternal(AdditionalFilter, NullCopyable, 0, InCompleteDelegate);
+}
+
+void UGameFeaturesSubsystem::LoadBuiltInGameFeaturePlugins_Amortized(const FBuiltInPluginAdditionalFilters_Copyable& AdditionalFilter_Copyable, int32 AmortizeRate, const FBuiltInGameFeaturePluginsLoaded& InCompleteDelegate /*= FBuiltInGameFeaturePluginsLoaded()*/)
+{
+	LoadBuiltInGameFeaturePluginsInternal(AdditionalFilter_Copyable, AdditionalFilter_Copyable, AmortizeRate, InCompleteDelegate);
+}
+
+void UGameFeaturesSubsystem::LoadBuiltInGameFeaturePluginsInternal(FBuiltInPluginAdditionalFilters AdditionalFilter, const FBuiltInPluginAdditionalFilters_Copyable& AdditionalFilter_Copyable, int32 AmortizeRate, const FBuiltInGameFeaturePluginsLoaded& InCompleteDelegate /*= FBuiltInGameFeaturePluginsLoaded()*/)
 {
 	TRACE_CPUPROFILER_EVENT_SCOPE_STR("UGameFeaturesSubsystem::LoadBuiltInGameFeaturePlugins");
 
@@ -2257,10 +2268,12 @@ void UGameFeaturesSubsystem::LoadBuiltInGameFeaturePlugins(FBuiltInPluginAdditio
 	{
 		FScopeLogTime ScopeLogTime{TEXT("BuiltInGameFeaturePlugins loaded."), nullptr, FConditionalScopeLogTime::ScopeLog_Seconds};
 
+		TArray<TSharedRef<IPlugin>> AmortizedPluginsRemaining;
 		TMap<FString, UE::GameFeatures::FResult> Results;
 		FBuiltInGameFeaturePluginsLoaded CompleteDelegate;
 
 		int32 NumPluginsLoaded = 0;
+		int32 PluginAmortizeRate = 0;
 		bool bPushedTagsBroadcast = false;
 		bool bPushedAssetBulkScanning = false;
 
@@ -2331,7 +2344,7 @@ void UGameFeaturesSubsystem::LoadBuiltInGameFeaturePlugins(FBuiltInPluginAdditio
 	}
 
 	const int32 NumPluginsToLoad = EnabledPlugins.Num();
-	UE_LOG(LogGameFeatures, Log, TEXT("Loading %i builtins"), NumPluginsToLoad);
+	UE_LOG(LogGameFeatures, Log, TEXT("Loading %i builtins. %s"), NumPluginsToLoad, AmortizeRate > 0 ? TEXT("Amortized") : TEXT("Not Amortized"));
 
 	// Sort the plugins so we can more accurately track how long it takes to load rather than have inconsistent dependency timings.
 	TArray<TSharedRef<IPlugin>> Dependencies;
@@ -2357,15 +2370,46 @@ void UGameFeaturesSubsystem::LoadBuiltInGameFeaturePlugins(FBuiltInPluginAdditio
 
 	Algo::TopologicalSort(EnabledPlugins, GetPluginDependencies);
 
-	for (const TSharedRef<IPlugin>& Plugin : EnabledPlugins)
+	auto ProcessGameFeaturePlugin = [this](const TSharedRef<IPlugin>& PluginToProcess, const TSharedRef<FLoadContext>& LoadContextToProcess, const FBuiltInPluginAdditionalFilters& AdditionalFilterForProcess)
 	{
-		FBuiltInPluginLoadTimeTrackerScope TrackerScope(PluginLoadTimeTracker, Plugin);
-		LoadBuiltInGameFeaturePlugin(Plugin, AdditionalFilter, FGameFeaturePluginLoadComplete::CreateLambda([LoadContext, Plugin](const UE::GameFeatures::FResult& Result)
+		LoadBuiltInGameFeaturePlugin(PluginToProcess, AdditionalFilterForProcess, FGameFeaturePluginLoadComplete::CreateLambda([LoadContextToProcess, PluginToProcess](const UE::GameFeatures::FResult& Result)
 		{
-			LoadContext->Results.Add(Plugin->GetName(), Result);
-			++LoadContext->NumPluginsLoaded;
-			UE_LOG(LogGameFeatures, VeryVerbose, TEXT("Finished Loading %i builtins"), LoadContext->NumPluginsLoaded);
+			LoadContextToProcess->Results.Add(PluginToProcess->GetName(), Result);
+			++LoadContextToProcess->NumPluginsLoaded;
+			UE_LOG(LogGameFeatures, VeryVerbose, TEXT("Finished Loading %i builtins"), LoadContextToProcess->NumPluginsLoaded);
 		}));
+	};
+
+	if (AmortizeRate > 0)
+	{
+		LoadContext->PluginAmortizeRate = AmortizeRate;
+		LoadContext->AmortizedPluginsRemaining = EnabledPlugins;
+		// Note that we are using AdditionalFilter_Copyable here. We need to make a copy of thus function so that the captured properties in lambdas are not destroyed
+		FTSTicker::GetCoreTicker().AddTicker(FTickerDelegate::CreateWeakLambda(this, [this, LoadContext, AdditionalFilter_Copyable, ProcessGameFeaturePlugin](float)
+		{
+			if (GameSpecificPolicies)
+			{
+				int32 NumPluginsThisFrame = FMath::Min(LoadContext->PluginAmortizeRate, LoadContext->AmortizedPluginsRemaining.Num());
+				
+				for (int32 PluginIdx = 0; PluginIdx < NumPluginsThisFrame; ++PluginIdx)
+				{
+					const TSharedRef<IPlugin>& Plugin = LoadContext->AmortizedPluginsRemaining[PluginIdx];
+					ProcessGameFeaturePlugin(Plugin, LoadContext, AdditionalFilter_Copyable);
+					
+				}
+
+				LoadContext->AmortizedPluginsRemaining.RemoveAt(0, NumPluginsThisFrame);
+			}
+			return GameSpecificPolicies && LoadContext->AmortizedPluginsRemaining.Num() > 0;
+		}));
+	}
+	else
+	{
+		for (const TSharedRef<IPlugin>& Plugin : EnabledPlugins)
+		{
+			FBuiltInPluginLoadTimeTrackerScope TrackerScope(PluginLoadTimeTracker, Plugin);
+			ProcessGameFeaturePlugin(Plugin, LoadContext, AdditionalFilter);
+		}
 	}
 }
 
diff --git a/Engine/Plugins/Runtime/GameFeatures/Source/GameFeatures/Public/GameFeaturesSubsystem.h b/Engine/Plugins/Runtime/GameFeatures/Source/GameFeatures/Public/GameFeaturesSubsystem.h
--- a/Engine/Plugins/Runtime/GameFeatures/Source/GameFeatures/Public/GameFeaturesSubsystem.h
+++ b/Engine/Plugins/Runtime/GameFeatures/Source/GameFeatures/Public/GameFeaturesSubsystem.h
@@ -631,6 +631,7 @@ public:
 	}
 
 	typedef TFunctionRef<bool(const FString& PluginFilename, const FGameFeaturePluginDetails& Details, FBuiltInGameFeaturePluginBehaviorOptions& OutOptions)> FBuiltInPluginAdditionalFilters;
+	typedef TFunction<bool(const FString& PluginFilename, const FGameFeaturePluginDetails& Details, FBuiltInGameFeaturePluginBehaviorOptions& OutOptions)> FBuiltInPluginAdditionalFilters_Copyable;
 
 	/** Loads a built-in game feature plugin if it passes the specified filter */
 	UE_API void LoadBuiltInGameFeaturePlugin(const TSharedRef<IPlugin>& Plugin, FBuiltInPluginAdditionalFilters AdditionalFilter, const FGameFeaturePluginLoadComplete& CompleteDelegate = FGameFeaturePluginLoadComplete());
@@ -638,6 +639,12 @@ public:
 	/** Loads all built-in game feature plugins that pass the specified filters */
 	UE_API void LoadBuiltInGameFeaturePlugins(FBuiltInPluginAdditionalFilters AdditionalFilter, const FBuiltInGameFeaturePluginsLoaded& CompleteDelegate = FBuiltInGameFeaturePluginsLoaded());
 
+	/** Loads all built-in game feature plugins that pass the specified filters, split over multiple frames processing only AmortizeRate plugins per frame if greater than 0. Note that AdditionalFilter must be a TFunction and not a TFunctionRef since it will be used in future ticks */
+	UE_API void LoadBuiltInGameFeaturePlugins_Amortized(const FBuiltInPluginAdditionalFilters_Copyable& AdditionalFilter_Copyable, int32 AmortizeRate, const FBuiltInGameFeaturePluginsLoaded& CompleteDelegate = FBuiltInGameFeaturePluginsLoaded());
+private:
+	void LoadBuiltInGameFeaturePluginsInternal(FBuiltInPluginAdditionalFilters AdditionalFilter, const FBuiltInPluginAdditionalFilters_Copyable& AdditionalFilter_Copyable, int32 AmortizeRate, const FBuiltInGameFeaturePluginsLoaded& CompleteDelegate = FBuiltInGameFeaturePluginsLoaded());
+public:
+
 	/** Returns the list of plugin filenames that have progressed beyond installed. Used in cooking to determine which will be cooked. */
 	//@TODO: GameFeaturePluginEnginePush: Might not be general enough for engine level, TBD
 	UE_API void GetLoadedGameFeaturePluginFilenamesForCooking(TArray<FString>& OutLoadedPluginFilenames) const;
