GFP - suppress DownloadNotAllowed error logs as these aren't errors in the traditional sense.

diff --git a/Engine/Plugins/Runtime/GameFeatures/Source/GameFeatures/Private/GameFeaturePluginStateMachine.cpp b/Engine/Plugins/Runtime/GameFeatures/Source/GameFeatures/Private/GameFeaturePluginStateMachine.cpp
--- a/Engine/Plugins/Runtime/GameFeatures/Source/GameFeatures/Private/GameFeaturePluginStateMachine.cpp
+++ b/Engine/Plugins/Runtime/GameFeatures/Source/GameFeatures/Private/GameFeaturePluginStateMachine.cpp
@@ -535,9 +535,8 @@ void FGameFeaturePluginStateStatus::SetTransition(EGameFeaturePluginState InTran
 	TransitionResult.OptionalErrorText = FText();
 }
 
-void FGameFeaturePluginStateStatus::SetTransitionError(EGameFeaturePluginState TransitionToErrorState, UE::GameFeatures::FResult TransitionResultIn)
+void FGameFeaturePluginStateStatus::SetTransitionError(EGameFeaturePluginState TransitionToErrorState, UE::GameFeatures::FResult TransitionResultIn, bool bInSuppressErrorLog /*= false*/)
 {
-	TransitionToState = TransitionToErrorState;
 	if (ensureAlwaysMsgf(TransitionResultIn.HasError(), TEXT("Invalid call to SetTransitionError with an FResult that isn't an error! TransitionToErrorState: %s"), *UE::GameFeatures::ToString(TransitionToErrorState)))
 	{
 		TransitionResult = MoveTemp(TransitionResultIn);
@@ -547,6 +546,9 @@ void FGameFeaturePluginStateStatus::SetTransitionError(EGameFeaturePluginState T
 		//Logic error using a non-error FResult, so just generate a general error to keep the SetTransitionError intent
 		TransitionResult = MakeError(TEXT("Invalid_Transition_Error"));
 	}
+
+	TransitionToState = TransitionToErrorState;
+	bSuppressErrorLog = bInSuppressErrorLog;
 }
 
 UE::GameFeatures::FResult FGameFeaturePluginState::GetErrorResult(const FString& ErrorCode, const FText OptionalErrorText/*= FText()*/) const
@@ -1525,6 +1527,7 @@ struct FGameFeaturePluginState_Downloading : public FGameFeaturePluginState
 	}
 
 	UE::GameFeatures::FResult Result;
+	bool bSuppressResultErrorLog = false;
 	bool bPluginDownloaded = false;
 
 	TArray<FName> PendingBundleDownloads;
@@ -1576,6 +1579,7 @@ struct FGameFeaturePluginState_Downloading : public FGameFeaturePluginState
 		IInstallBundleManager::PausedBundleDelegate.RemoveAll(this);
 
 		Result = MakeValue();
+		bSuppressResultErrorLog = false;
 		bPluginDownloaded = false;
 		PendingBundleDownloads.Empty();
 		ProgressTracker = nullptr;
@@ -1842,6 +1846,7 @@ struct FGameFeaturePluginState_Downloading : public FGameFeaturePluginState
 				if (!InstallState.GetAllBundlesHaveState(EInstallBundleInstallState::UpToDate))
 				{
 					Result = GetErrorResult(TEXT("GFPStateMachine.DownloadNotAllowed"));
+					bSuppressResultErrorLog = true; // Don't log an error if the user disallowed the download
 					UpdateStateMachineImmediate();
 					return;
 				}
@@ -1865,6 +1870,7 @@ struct FGameFeaturePluginState_Downloading : public FGameFeaturePluginState
 			if (Options.bDoNotDownload && !InstallState.GetAllBundlesHaveState(EInstallBundleInstallState::UpToDate))
 			{
 				Result = GetErrorResult(TEXT("GFPStateMachine.DownloadNotAllowed"));
+				bSuppressResultErrorLog = true; // Don't log an error if the user disallowed the download
 				UpdateStateMachineImmediate();
 				return;
 			}
@@ -1893,7 +1899,7 @@ struct FGameFeaturePluginState_Downloading : public FGameFeaturePluginState
 		if (!Result.HasValue())
 		{
 			const EGameFeaturePluginState FailState = GetFailureTransitionState();
-			StateStatus.SetTransitionError(FailState, Result);
+			StateStatus.SetTransitionError(FailState, Result, bSuppressResultErrorLog);
 			return;
 		}
 
@@ -4469,12 +4475,25 @@ void UGameFeaturePluginStateMachine::UpdateStateMachine()
 		{
 			if (!StateStatus.TransitionResult.HasValue())
 			{
-				UE_LOG(LogGameFeatures, Error, TEXT("Game feature '%s' transition failed. Ending state: %s [%s, %s]. Result: %s"),
-					*GetGameFeatureName(),
-					*UE::GameFeatures::ToString(CurrentState),
-					*UE::GameFeatures::ToString(StateProperties.Destination.MinState),
-					*UE::GameFeatures::ToString(StateProperties.Destination.MaxState),
-					*UE::GameFeatures::ToString(StateStatus.TransitionResult));
+				constexpr static const TCHAR ErrLogFmt[] = TEXT("Game feature '%s' transition failed. Ending state: %s [%s, %s]. Result: %s");
+				if (StateStatus.bSuppressErrorLog)
+				{
+					UE_LOG(LogGameFeatures, Display, ErrLogFmt,
+						*GetGameFeatureName(),
+						*UE::GameFeatures::ToString(CurrentState),
+						*UE::GameFeatures::ToString(StateProperties.Destination.MinState),
+						*UE::GameFeatures::ToString(StateProperties.Destination.MaxState),
+						*UE::GameFeatures::ToString(StateStatus.TransitionResult));
+				}
+				else
+				{
+					UE_LOG(LogGameFeatures, Error, ErrLogFmt,
+						*GetGameFeatureName(),
+						*UE::GameFeatures::ToString(CurrentState),
+						*UE::GameFeatures::ToString(StateProperties.Destination.MinState),
+						*UE::GameFeatures::ToString(StateProperties.Destination.MaxState),
+						*UE::GameFeatures::ToString(StateStatus.TransitionResult));
+				}
 			}
 			else if (StateProperties.Destination.Contains(CurrentState))
 			{
diff --git a/Engine/Plugins/Runtime/GameFeatures/Source/GameFeatures/Private/GameFeaturePluginStateMachine.h b/Engine/Plugins/Runtime/GameFeatures/Source/GameFeatures/Private/GameFeaturePluginStateMachine.h
--- a/Engine/Plugins/Runtime/GameFeatures/Source/GameFeatures/Private/GameFeaturePluginStateMachine.h
+++ b/Engine/Plugins/Runtime/GameFeatures/Source/GameFeatures/Private/GameFeaturePluginStateMachine.h
@@ -354,17 +354,20 @@ struct FGameFeaturePluginStateMachineProperties
 struct FGameFeaturePluginStateStatus
 {
 private:
+	/** Holds the current error for any state transition. */
+	UE::GameFeatures::FResult TransitionResult = MakeValue();
+
 	/** The state to transition to after UpdateState is complete. */
 	EGameFeaturePluginState TransitionToState = EGameFeaturePluginState::Uninitialized;
 
-	/** Holds the current error for any state transition. */
-	UE::GameFeatures::FResult TransitionResult = MakeValue();
+	/** Whether to suppress error logging if TransitionResult is an error. */
+	bool bSuppressErrorLog = false;
 
 	friend class UGameFeaturePluginStateMachine;
 
 public:
 	void SetTransition(EGameFeaturePluginState InTransitionToState);
-	void SetTransitionError(EGameFeaturePluginState TransitionToErrorState, UE::GameFeatures::FResult TransitionResult);
+	void SetTransitionError(EGameFeaturePluginState TransitionToErrorState, UE::GameFeatures::FResult TransitionResult, bool bInSuppressErrorLog = false);
 };
 
 enum class EGameFeaturePluginStateType : uint8
