Change PlayStation BootTest Errors to Warnings

#tests PS4 Boot Test passes with this change.
#rnx

diff --git a/Engine/Plugins/Runtime/GameFeatures/Source/GameFeatures/Private/GameFeaturesSubsystem.cpp b/Engine/Plugins/Runtime/GameFeatures/Source/GameFeatures/Private/GameFeaturesSubsystem.cpp
--- a/Engine/Plugins/Runtime/GameFeatures/Source/GameFeatures/Private/GameFeaturesSubsystem.cpp
+++ b/Engine/Plugins/Runtime/GameFeatures/Source/GameFeatures/Private/GameFeaturesSubsystem.cpp
@@ -3207,13 +3207,13 @@ bool UGameFeaturesSubsystem::FindOrCreatePluginDependencyStateMachines(const FSt
 			if (DependencyURLInfo.HasError())
 			{
 				FString ErrorMessage = FString::Printf(TEXT("Game feature plugin '%s' has unknown dependency '%s' [%s]."), *PluginFilename, *DependencyName, *DependencyURLInfo.GetError());
-				UE_LOG(LogGameFeatures, Error, TEXT("%s"), *ErrorMessage);
 
 				SetExplanationForNotMountingPlugin(PluginURL, ErrorMessage);
 
 				if (DependencyDetails.bFailIfNotFound)
 				{
 					// This plugin is known to be a hard dependency so fail early before something harder to debug happens.
+					UE_LOG(LogGameFeatures, Error, TEXT("%s"), *ErrorMessage);
 					return false;
 				}
 				else
@@ -3221,6 +3221,7 @@ bool UGameFeaturesSubsystem::FindOrCreatePluginDependencyStateMachines(const FSt
 					// Don't actually return false here as we want to still be able to progress in the case of 
 					// things like an editor plugin being included as a dependency in the client. We can't tell from just 
 					// the reference if a plugin is not enabled for this build target.
+					UE_LOG(LogGameFeatures, Log, TEXT("%s"), *ErrorMessage);
 					continue;
 				}
 			}
