Guard against dangling `this` captures in some async loading callbacks in the GFP state machine, and verify the callback is for the expected handle

diff --git a/Engine/Plugins/Runtime/GameFeatures/Source/GameFeatures/Private/GameFeaturePluginStateMachine.cpp b/Engine/Plugins/Runtime/GameFeatures/Source/GameFeatures/Private/GameFeaturePluginStateMachine.cpp
--- a/Engine/Plugins/Runtime/GameFeatures/Source/GameFeatures/Private/GameFeaturePluginStateMachine.cpp
+++ b/Engine/Plugins/Runtime/GameFeatures/Source/GameFeatures/Private/GameFeaturePluginStateMachine.cpp
@@ -3152,8 +3152,17 @@ struct FGameFeaturePluginState_Registering : public FGameFeaturePluginState
 			GameFeatureDataHandle = UGameFeaturesSubsystem::LoadGameFeatureData(Path, true /*bStartStalled*/);
 			if (GameFeatureDataHandle && GameFeatureDataHandle->IsLoadingInProgress())
 			{
-				GameFeatureDataHandle->BindCancelDelegate(FStreamableDelegate::CreateLambda([this]
+				// CurrentMachine owns `this`, so use it as a lifetime check for the `this` captured in the lambdas below, as they may be called after `this` is destroyed
+				UGameFeaturePluginStateMachine* CurrentMachine = UGameFeaturesSubsystem::Get().FindGameFeaturePluginStateMachine(StateProperties.PluginIdentifier);
+
+				GameFeatureDataHandle->BindCancelDelegate(FStreamableDelegateWithHandle::CreateWeakLambda(CurrentMachine, [this](const TSharedPtr<FStreamableHandle>& Handle)
 				{
+					if (GameFeatureDataHandle != Handle)
+					{
+						// We're no longer in a state where we want to process this callback (maybe we already went through EndState)
+						return;
+					}
+
 					const FStringView ShortUrl = StateProperties.PluginIdentifier.GetIdentifyingString();
 					UE_LOG(LogGameFeatures, Error, TEXT("Game Feature Data loading was cancelled for URL %.*s"), ShortUrl.Len(), ShortUrl.GetData());
 
@@ -3161,8 +3170,14 @@ struct FGameFeaturePluginState_Registering : public FGameFeaturePluginState
 					UpdateStateMachineDeferred();
 				}));
 
-				GameFeatureDataHandle->BindCompleteDelegate(FStreamableDelegate::CreateLambda([this, Attempt]
+				GameFeatureDataHandle->BindCompleteDelegate(FStreamableDelegateWithHandle::CreateWeakLambda(CurrentMachine, [this, Attempt](const TSharedPtr<FStreamableHandle>& Handle)
 				{
+					if (GameFeatureDataHandle != Handle)
+					{
+						// We're no longer in a state where we want to process this callback (maybe we already went through EndState)
+						return;
+					}
+
 					StateProperties.GameFeatureData = Cast<UGameFeatureData>(GameFeatureDataHandle->GetLoadedAsset());
 					if (!StateProperties.GameFeatureData)
 					{
@@ -3444,7 +3459,7 @@ struct FGameFeaturePluginState_Loading : public FGameFeaturePluginState
 			UE::GameFeatures::MountLocalizationData(CurrentMachine, StateProperties);
 		}
 
-		BundleHandle = LoadGameFeatureBundles(StateProperties.GameFeatureData);
+		LoadGameFeatureBundles(StateProperties.GameFeatureData);
 	}
 
 	virtual void EndState() override
@@ -3483,7 +3498,7 @@ struct FGameFeaturePluginState_Loading : public FGameFeaturePluginState
 	}
 
 	/** Loads primary assets and bundles for the specified game feature */
-	TSharedPtr<FStreamableHandle> LoadGameFeatureBundles(const UGameFeatureData* GameFeatureToLoad)
+	void LoadGameFeatureBundles(const UGameFeatureData* GameFeatureToLoad)
 	{
 		check(GameFeatureToLoad);
 
@@ -3497,27 +3512,43 @@ struct FGameFeaturePluginState_Loading : public FGameFeaturePluginState
 			AssetIdsToLoad.Add(GameFeatureAssetId);
 		}
 
-		TSharedPtr<FStreamableHandle> RetHandle;
 		if (AssetIdsToLoad.Num() > 0)
 		{
+			// CurrentMachine owns `this`, so use it as a lifetime check for the `this` captured in the lambdas below, as they may be called after `this` is destroyed
+			UGameFeaturePluginStateMachine* CurrentMachine = UGameFeaturesSubsystem::Get().FindGameFeaturePluginStateMachine(StateProperties.PluginIdentifier);
+
 			FAssetManagerLoadParams LoadParams;
-			LoadParams.OnCancel = FStreamableDelegateWithHandle::CreateLambda([this](const TSharedPtr<FStreamableHandle>&)
+			LoadParams.OnCancel = FStreamableDelegateWithHandle::CreateWeakLambda(CurrentMachine, [this](const TSharedPtr<FStreamableHandle>& Handle)
 			{
+				if (BundleHandle != Handle)
+				{
+					// We're no longer in a state where we want to process this callback (maybe we already went through EndState)
+					return;
+				}
+
 				const FStringView ShortUrl = StateProperties.PluginIdentifier.GetIdentifyingString();
 				UE_LOG(LogGameFeatures, Error, TEXT("Game Feature preloading was cancelled for URL %.*s"), ShortUrl.Len(), ShortUrl.GetData());
 				UpdateStateMachineDeferred();
 			});
 
 			// This can't be bound to the handle after its created, AM may bind it internally
-			LoadParams.OnComplete = FStreamableDelegateWithHandle::CreateLambda([this](const TSharedPtr<FStreamableHandle>&)
+			LoadParams.OnComplete = FStreamableDelegateWithHandle::CreateWeakLambda(CurrentMachine, [this](const TSharedPtr<FStreamableHandle>& Handle)
 			{
+				if (BundleHandle != Handle)
+				{
+					// We're no longer in a state where we want to process this callback (maybe we already went through EndState)
+					return;
+				}
+
 				UpdateStateMachineDeferred();
 			});
 
-			RetHandle = UAssetManager::Get().LoadPrimaryAssets(AssetIdsToLoad, Policy.GetPreloadBundleStateForGameFeature(), MoveTemp(LoadParams));
+			BundleHandle = UAssetManager::Get().LoadPrimaryAssets(AssetIdsToLoad, Policy.GetPreloadBundleStateForGameFeature(), MoveTemp(LoadParams));
+		}
+		else
+		{
+			BundleHandle.Reset();
 		}
-
-		return RetHandle;
 	}
 };
 
