Add a way to generate GFPs programmatically for the purpose of testing
CanContainContent: false allows us to simply generate the uplugin, and when loading the GFP the GameFeatureData will be generated by the GFP system
Create a test that tests our current failure of A -> C, D and B -> C, D, activate A and B, then move B to registered causing C and D to downgrade
Leave test enabled to test CreatePlugin, but the test doesnt actually test the failure it should yet (until fixed by ref count!)

diff --git a/Engine/Plugins/Runtime/GameFeatures/Source/GameFeatures/Private/Tests/GameFeaturePluginTestHelper.cpp b/Engine/Plugins/Runtime/GameFeatures/Source/GameFeatures/Private/Tests/GameFeaturePluginTestHelper.cpp
new file mode 100644
--- /dev/null
+++ b/Engine/Plugins/Runtime/GameFeatures/Source/GameFeatures/Private/Tests/GameFeaturePluginTestHelper.cpp
@@ -0,0 +1,73 @@
+// Copyright Epic Games, Inc. All Rights Reserved.
+
+#if WITH_DEV_AUTOMATION_TESTS && WITH_EDITOR
+
+#include "GameFeaturePluginTestsHelper.h"
+#include "GameFeaturesSubsystemSettings.h"
+#include "Misc/FileHelper.h"
+#include "Misc/Paths.h"
+#include "HAL/FileManager.h"
+
+FString GeneratePluginString(const TArray<FGameFeatureDependsProperties>& Depends)
+{
+	FString Out;
+	for (const FGameFeatureDependsProperties& Depend : Depends)
+	{
+		Out += FString::Printf(TEXT(R"(
+		{
+			"Name": "%s",
+			"Enabled": true,
+			"Activate": %s,
+            "Optional": true
+		},
+
+		)"), *Depend.PluginName, Depend.ShouldActivate == EShouldActivate::No ? TEXT("false") : TEXT("true"));
+	}
+
+	return Out;
+}
+
+bool CreateGameFeaturePlugin(FGameFeatureProperties Properties, FString& OutPluginURL)
+{
+	FString PluginPath = FPaths::ProjectPluginsDir() / TEXT("GameFeatures") / Properties.PluginName;
+	FString UPluginPath = PluginPath / (Properties.PluginName + TEXT(".uplugin"));
+
+	// if one was already created here, delete it before making a new one
+	if (FPaths::DirectoryExists(PluginPath))
+	{
+		IFileManager::Get().DeleteDirectory(*PluginPath, /* bEnsureExists */ false, /* bDeleteEntireTree */ true);
+	}
+
+	IFileManager::Get().MakeDirectory(*PluginPath, /* bCreateTree */ true);
+
+	FString PluginDependsString = GeneratePluginString(Properties.Depends);
+	FString PluginDetials = FString::Printf(TEXT(R"(
+	{
+		"FileVersion": 3,
+		"Version": 1,
+		"VersionName": "1.0",
+		"FriendlyName": "%s",
+		"Description": "Generated GFP for Testing",
+		"Category": "GFPTesting",
+		"CreatedBy": "Automated",
+		"CreatedByURL": "",
+		"DocsURL": "",
+		"MarketplaceURL": "",
+		"SupportURL": "",
+		"EnabledByDefault": false,
+		"CanContainContent": false,
+		"IsBetaVersion": false,
+		"IsExperimentalVersion": false,
+		"Installed": false,
+		"ExplicitlyLoaded": true,
+		"BuiltInInitialFeatureState": "%s",
+		"Plugins": [%s]
+	})"), *Properties.PluginName, *LexToString(Properties.BuiltinAutoState), *PluginDependsString);
+
+	bool bSavedFile = FFileHelper::SaveStringToFile(PluginDetials, *UPluginPath);
+
+	OutPluginURL = UGameFeaturesSubsystem::GetPluginURL_FileProtocol(IFileManager::Get().ConvertToRelativePath(*UPluginPath));
+	return bSavedFile;
+}
+
+#endif
diff --git a/Engine/Plugins/Runtime/GameFeatures/Source/GameFeatures/Private/Tests/GameFeaturePluginTests.cpp b/Engine/Plugins/Runtime/GameFeatures/Source/GameFeatures/Private/Tests/GameFeaturePluginTests.cpp
--- a/Engine/Plugins/Runtime/GameFeatures/Source/GameFeatures/Private/Tests/GameFeaturePluginTests.cpp
+++ b/Engine/Plugins/Runtime/GameFeatures/Source/GameFeatures/Private/Tests/GameFeaturePluginTests.cpp
@@ -9,6 +9,7 @@
 
 #include "GameFeaturePluginOperationResult.h"
 #include "GameFeaturePluginStateMachine.h"
+#include "GameFeaturePluginTestsHelper.h"
 #include "GameFeaturesSubsystem.h"
 #include "GameFeatureTypes.h"
 
@@ -55,33 +56,38 @@ public:
 	{
 	}
 
-	bool IsPluginInPluginStateRange(const FGameFeaturePluginStateRange PluginStateRange)
+	bool IsPluginInPluginStateRange(const FGameFeaturePluginStateRange PluginStateRange, const FString& PluginURL)
 	{
-		EGameFeaturePluginState CurrentPluginState = UGameFeaturesSubsystem::Get().GetPluginState(GFPFileURL);
+		EGameFeaturePluginState CurrentPluginState = UGameFeaturesSubsystem::Get().GetPluginState(PluginURL);
 
 		return PluginStateRange.Contains(CurrentPluginState);
 	}
 
-	void LatentTestPluginState(const FGameFeaturePluginStateRange PluginStateRange)
+	void LatentTestPluginState(const EGameFeatureTargetState PluginTargetState, const FString& PluginURL)
 	{
-		ADD_LATENT_AUTOMATION_COMMAND(FExecuteFunction([this, PluginStateRange]
+		LatentTestPluginState(FGameFeaturePluginStateRange(ConvertTargetStateToPluginState(PluginTargetState)), PluginURL);
+	}
+
+	void LatentTestPluginState(const FGameFeaturePluginStateRange PluginStateRange, const FString& PluginURL)
+	{
+		ADD_LATENT_AUTOMATION_COMMAND(FExecuteFunction([this, PluginStateRange, PluginURL]
 		{
-			TestTrue(FString::Printf(TEXT("Plugin in %s state, expected plugin state in range (%s, %s)"),
-				*UE::GameFeatures::ToString(UGameFeaturesSubsystem::Get().GetPluginState(GFPFileURL)), *UE::GameFeatures::ToString(PluginStateRange.MinState), *UE::GameFeatures::ToString(PluginStateRange.MaxState)),
-				IsPluginInPluginStateRange(PluginStateRange));
+			TestTrue(FString::Printf(TEXT("Plugin %s in %s state, expected plugin state in range (%s, %s)"),
+				*PluginURL, *UE::GameFeatures::ToString(UGameFeaturesSubsystem::Get().GetPluginState(PluginURL)), *UE::GameFeatures::ToString(PluginStateRange.MinState), *UE::GameFeatures::ToString(PluginStateRange.MaxState)),
+				IsPluginInPluginStateRange(PluginStateRange, PluginURL));
 
 			return true;
 		}));
 	}
 
-	void LatentTestTransitionGFP(const EGameFeatureTargetState TargetState)
+	void LatentTestTransitionGFP(const EGameFeatureTargetState TargetState, const FString& PluginURL)
 	{
-		ADD_LATENT_AUTOMATION_COMMAND(FExecuteFunction([this, TargetState]
+		ADD_LATENT_AUTOMATION_COMMAND(FExecuteFunction([this, TargetState, PluginURL]
 		{
 			*bAsyncCommandComplete = false;
 
-			UGameFeaturesSubsystem::Get().ChangeGameFeatureTargetState(GFPFileURL, TargetState,
-				FGameFeaturePluginChangeStateComplete::CreateLambda([this, TargetState](const UE::GameFeatures::FResult& Result)
+			UGameFeaturesSubsystem::Get().ChangeGameFeatureTargetState(PluginURL, TargetState,
+				FGameFeaturePluginChangeStateComplete::CreateLambda([this, TargetState, PluginURL](const UE::GameFeatures::FResult& Result)
 				{
 					*bAsyncCommandComplete = true;
 					TestFalse(FString::Printf(TEXT("Failed to transition to %s: error: %s"), *LexToString(TargetState), *UE::GameFeatures::ToString(Result)),
@@ -94,14 +100,14 @@ public:
 
 		ADD_LATENT_AUTOMATION_COMMAND(FWaitForTrue(&*bAsyncCommandComplete));
 
-		LatentTestPluginState(FGameFeaturePluginStateRange(ConvertTargetStateToPluginState(TargetState)));
+		LatentTestPluginState(TargetState, PluginURL);
 	}
 
 	void LatentCheckInitialPluginState()
 	{
 		// Check we are somewhere between uninited, and uninstalled for the first time we check this and after we restore the plugin state
 		// depending on the initial state as well as deactivating/terminating the plugin we should be in the Terminal or UnknownStatus node
-		LatentTestPluginState(FGameFeaturePluginStateRange(EGameFeaturePluginState::Uninitialized, EGameFeaturePluginState::Uninstalled));
+		LatentTestPluginState(FGameFeaturePluginStateRange(EGameFeaturePluginState::Uninitialized, EGameFeaturePluginState::Uninstalled), GFPFileURL);
 	}
 
 	void LatentRestorePluginState()
@@ -111,7 +117,7 @@ public:
 			*bAsyncCommandComplete = false;
 
 			// We are in an uninstalled/terminal/not setup state. Dont try to Deactivate/Terminate when we are not Activated/Installed
-			if (IsPluginInPluginStateRange(FGameFeaturePluginStateRange(EGameFeaturePluginState::Uninitialized, EGameFeaturePluginState::Uninstalled)))
+			if (IsPluginInPluginStateRange(FGameFeaturePluginStateRange(EGameFeaturePluginState::Uninitialized, EGameFeaturePluginState::Uninstalled), GFPFileURL))
 			{
 				*bAsyncCommandComplete = true;
 				return true;
@@ -135,7 +141,7 @@ public:
 			*bAsyncCommandComplete = false;
 
 			// We are in an uninstalled/terminal/not setup state. Dont try to Deactivate/Terminate when we are not Activated/Installed
-			if (IsPluginInPluginStateRange(FGameFeaturePluginStateRange(EGameFeaturePluginState::Uninitialized, EGameFeaturePluginState::Uninstalled)))
+			if (IsPluginInPluginStateRange(FGameFeaturePluginStateRange(EGameFeaturePluginState::Uninitialized, EGameFeaturePluginState::Uninstalled), GFPFileURL))
 			{
 				*bAsyncCommandComplete = true;
 				return true;
@@ -156,8 +162,8 @@ public:
 	}
 
 	// For now hard-coded into EngineTest area but can always be adjusted later
-	FString GFPPluginPath = TEXT("../../../EngineTest/Plugins/GameFeatures/GameFeatureEngineTestC/GameFeatureEngineTestC.uplugin");
-	FString GFPFileURL = FString(TEXT("file:")) + GFPPluginPath;
+	const FString GFPPluginPath = TEXT("../../../EngineTest/Plugins/GameFeatures/GameFeatureEngineTestC/GameFeatureEngineTestC.uplugin");
+	const FString GFPFileURL = FString(TEXT("file:")) + GFPPluginPath;
 	TSharedRef<bool> bAsyncCommandComplete = MakeShared<bool>(false);
 };
 
@@ -172,10 +178,10 @@ bool FGameFeatureSubsystemTestChangeState::RunTest(const FString& Parameters)
 		LatentRestorePluginState();
 	};
 
-	LatentTestTransitionGFP(EGameFeatureTargetState::Installed);
-	LatentTestTransitionGFP(EGameFeatureTargetState::Registered);
-	LatentTestTransitionGFP(EGameFeatureTargetState::Loaded);
-	LatentTestTransitionGFP(EGameFeatureTargetState::Active);
+	LatentTestTransitionGFP(EGameFeatureTargetState::Installed, GFPFileURL);
+	LatentTestTransitionGFP(EGameFeatureTargetState::Registered, GFPFileURL);
+	LatentTestTransitionGFP(EGameFeatureTargetState::Loaded, GFPFileURL);
+	LatentTestTransitionGFP(EGameFeatureTargetState::Active, GFPFileURL);
 
     return true;
 }
@@ -192,7 +198,7 @@ bool FGameFeatureSubsystemTestUninstall::RunTest(const FString& Parameters)
 	};
 
 	// Get us into an installed state so we can query info about the GFP
-	LatentTestTransitionGFP(EGameFeatureTargetState::Installed);
+	LatentTestTransitionGFP(EGameFeatureTargetState::Installed, GFPFileURL);
 
 	ADD_LATENT_AUTOMATION_COMMAND(FExecuteFunction([this]
 	{
@@ -246,7 +252,7 @@ bool FGameFeatureSubsystemTestGetGameFeatureData::RunTest(const FString& Paramet
 		LatentRestorePluginState();
 	};
 
-	LatentTestTransitionGFP(EGameFeatureTargetState::Installed);
+	LatentTestTransitionGFP(EGameFeatureTargetState::Installed, GFPFileURL);
 
 	ADD_LATENT_AUTOMATION_COMMAND(FExecuteFunction([this]
 	{
@@ -256,7 +262,7 @@ bool FGameFeatureSubsystemTestGetGameFeatureData::RunTest(const FString& Paramet
 		return true;
 	}));
 
-	LatentTestTransitionGFP(EGameFeatureTargetState::Registered);
+	LatentTestTransitionGFP(EGameFeatureTargetState::Registered, GFPFileURL);
 
 	ADD_LATENT_AUTOMATION_COMMAND(FExecuteFunction([this]
 	{
@@ -266,7 +272,7 @@ bool FGameFeatureSubsystemTestGetGameFeatureData::RunTest(const FString& Paramet
 		return true;
 	}));
 
-	LatentTestTransitionGFP(EGameFeatureTargetState::Active);
+	LatentTestTransitionGFP(EGameFeatureTargetState::Active, GFPFileURL);
 
 	ADD_LATENT_AUTOMATION_COMMAND(FExecuteFunction([this]
 	{
@@ -292,8 +298,8 @@ bool FGameFeatureSubsystemTestNonBuiltinPluginDoesntConvertToBuiltinPlugin::RunT
 		LatentRestorePluginState();
 	};
 
-	LatentTestTransitionGFP(EGameFeatureTargetState::Installed);
-	LatentTestTransitionGFP(EGameFeatureTargetState::Registered);
+	LatentTestTransitionGFP(EGameFeatureTargetState::Installed, GFPFileURL);
+	LatentTestTransitionGFP(EGameFeatureTargetState::Registered, GFPFileURL);
 
 	// Test we get to registered, installed -> mounted which will get our plugin in the enabled/mounted state
 	ADD_LATENT_AUTOMATION_COMMAND(FExecuteFunction([this]
@@ -338,4 +344,78 @@ bool FGameFeatureSubsystemTestNonBuiltinPluginDoesntConvertToBuiltinPlugin::RunT
 	return true;
 }
 
+
+/**
+ * This is a test that will fail currently:
+ *   Create GFPs A -> C, D and B -> C, D then activate A and B
+ *   move B to registered causing C, D to transition to Loaded (since they are ShouldActive)
+ *     expect:
+ *       C, D to stay active since A is still active and depends on that
+ *     result:
+ *       C, D become deactiving/loaded since B we downgraded to registered, which deactivates its depends
+ *
+ *   TODO this is simply testing our CreatePlugin logic for now, Leaving on even though it doesnt test the failure at the bottom!
+ */
+IMPLEMENT_CUSTOM_SIMPLE_AUTOMATION_TEST(FGameFeatureSubsystemTestCreatePlugin, FTestGameFeaturePluginBase, "GameFeaturePlugin.Subsystem.DeactivatePreventsDependsDowngradeIfRefCount", EAutomationTestFlags::EditorContext | EAutomationTestFlags::EngineFilter);
+bool FGameFeatureSubsystemTestCreatePlugin::RunTest(const FString& Parameters)
+{
+	FGameFeatureProperties PropertiesA;
+	PropertiesA.PluginName = TEXT("TestingA");
+	PropertiesA.BuiltinAutoState = EGameFeatureTargetState::Installed;
+	PropertiesA.Depends = { FGameFeatureDependsProperties{ TEXT("TestingD"), EShouldActivate::Yes }, FGameFeatureDependsProperties{ TEXT("TestingC"), EShouldActivate::Yes } };
+
+	FString PluginAURL;
+	CreateGameFeaturePlugin(PropertiesA, PluginAURL);
+
+	FGameFeatureProperties PropertiesB;
+	PropertiesB.PluginName = TEXT("TestingB");
+	PropertiesB.BuiltinAutoState = EGameFeatureTargetState::Installed;
+	PropertiesB.Depends = { FGameFeatureDependsProperties{ TEXT("TestingD"), EShouldActivate::Yes }, FGameFeatureDependsProperties{ TEXT("TestingC"), EShouldActivate::Yes } };
+
+	FString PluginBURL;
+	CreateGameFeaturePlugin(PropertiesB, PluginBURL);
+
+	FGameFeatureProperties PropertiesC;
+	PropertiesC.PluginName = TEXT("TestingC");
+	PropertiesC.BuiltinAutoState = EGameFeatureTargetState::Installed;
+
+	FString PluginCURL;
+	CreateGameFeaturePlugin(PropertiesC, PluginCURL);
+
+	FGameFeatureProperties PropertiesD;
+	PropertiesD.PluginName = TEXT("TestingD");
+	PropertiesD.BuiltinAutoState = EGameFeatureTargetState::Installed;
+
+	FString PluginDURL;
+	CreateGameFeaturePlugin(PropertiesD, PluginDURL);
+
+	// make sure they are all installed
+	LatentTestTransitionGFP(EGameFeatureTargetState::Installed, PluginCURL);
+	LatentTestTransitionGFP(EGameFeatureTargetState::Installed, PluginDURL);
+	LatentTestTransitionGFP(EGameFeatureTargetState::Installed, PluginAURL);
+	LatentTestTransitionGFP(EGameFeatureTargetState::Installed, PluginBURL);
+
+	// register A, and B which will pull in C, D to registered
+	LatentTestTransitionGFP(EGameFeatureTargetState::Registered, PluginAURL);
+	LatentTestTransitionGFP(EGameFeatureTargetState::Registered, PluginBURL);
+	LatentTestPluginState(EGameFeatureTargetState::Registered, PluginCURL);
+	LatentTestPluginState(EGameFeatureTargetState::Registered, PluginDURL);
+
+	// activate A, and B which will pull in C, D to registered
+	LatentTestTransitionGFP(EGameFeatureTargetState::Active, PluginAURL);
+	LatentTestTransitionGFP(EGameFeatureTargetState::Active, PluginBURL);
+	LatentTestPluginState(EGameFeatureTargetState::Active, PluginCURL);
+	LatentTestPluginState(EGameFeatureTargetState::Active, PluginDURL);
+
+	// drop A to registered, which currently drops C, D to Loaded which is a bug
+	LatentTestTransitionGFP(EGameFeatureTargetState::Registered, PluginAURL);
+
+	/* TODO remove this once we have ref counting working, and stop deactiving C and D here
+	LatentTestPluginState(EGameFeaturePluginState::Active, PluginCURL);
+	LatentTestPluginState(EGameFeaturePluginState::Active, PluginDURL);
+	*/
+
+	return true;
+}
+
 #endif // WITH_DEV_AUTOMATION_TESTS && WITH_EDITOR
diff --git a/Engine/Plugins/Runtime/GameFeatures/Source/GameFeatures/Private/Tests/GameFeaturePluginTestsHelper.h b/Engine/Plugins/Runtime/GameFeatures/Source/GameFeatures/Private/Tests/GameFeaturePluginTestsHelper.h
new file mode 100644
--- /dev/null
+++ b/Engine/Plugins/Runtime/GameFeatures/Source/GameFeatures/Private/Tests/GameFeaturePluginTestsHelper.h
@@ -0,0 +1,26 @@
+// Copyright Epic Games, Inc. All Rights Reserved.
+
+#pragma once
+
+#include "GameFeaturesSubsystem.h"
+
+enum class EShouldActivate : uint8
+{
+	No,
+	Yes
+};
+
+struct FGameFeatureDependsProperties
+{
+	FString PluginName;
+	EShouldActivate ShouldActivate = EShouldActivate::No;
+};
+
+struct FGameFeatureProperties
+{
+	FString PluginName;
+	EGameFeatureTargetState	BuiltinAutoState = EGameFeatureTargetState::Installed;
+	TArray<FGameFeatureDependsProperties> Depends;
+};
+
+bool CreateGameFeaturePlugin(FGameFeatureProperties Properties, FString& OutPluginURL);
