Use string builder when building up our plugins depends
Fix typo in comment

diff --git a/Engine/Plugins/Runtime/GameFeatures/Source/GameFeatures/Private/Tests/GameFeaturePluginTestHelper.cpp b/Engine/Plugins/Runtime/GameFeatures/Source/GameFeatures/Private/Tests/GameFeaturePluginTestHelper.cpp
--- a/Engine/Plugins/Runtime/GameFeatures/Source/GameFeatures/Private/Tests/GameFeaturePluginTestHelper.cpp
+++ b/Engine/Plugins/Runtime/GameFeatures/Source/GameFeatures/Private/Tests/GameFeaturePluginTestHelper.cpp
@@ -10,7 +10,7 @@
 
 FString GeneratePluginString(const TArray<FGameFeatureDependsProperties>& Depends)
 {
-	FString Out;
+	TStringBuilder<512> Out;
 	for (const FGameFeatureDependsProperties& Depend : Depends)
 	{
 		Out += FString::Printf(TEXT(R"(
@@ -18,13 +18,13 @@ FString GeneratePluginString(const TArray<FGameFeatureDependsProperties>& Depend
 			"Name": "%s",
 			"Enabled": true,
 			"Activate": %s,
-            "Optional": true
+			"Optional": true
 		},
 
 		)"), *Depend.PluginName, Depend.ShouldActivate == EShouldActivate::No ? TEXT("false") : TEXT("true"));
 	}
 
-	return Out;
+	return *Out;
 }
 
 bool CreateGameFeaturePlugin(FGameFeatureProperties Properties, FString& OutPluginURL)
diff --git a/Engine/Plugins/Runtime/GameFeatures/Source/GameFeatures/Private/Tests/GameFeaturePluginTests.cpp b/Engine/Plugins/Runtime/GameFeatures/Source/GameFeatures/Private/Tests/GameFeaturePluginTests.cpp
--- a/Engine/Plugins/Runtime/GameFeatures/Source/GameFeatures/Private/Tests/GameFeaturePluginTests.cpp
+++ b/Engine/Plugins/Runtime/GameFeatures/Source/GameFeatures/Private/Tests/GameFeaturePluginTests.cpp
@@ -401,7 +401,7 @@ bool FGameFeatureSubsystemTestCreatePlugin::RunTest(const FString& Parameters)
 	LatentTestPluginState(EGameFeatureTargetState::Registered, PluginCURL);
 	LatentTestPluginState(EGameFeatureTargetState::Registered, PluginDURL);
 
-	// activate A, and B which will pull in C, D to registered
+	// activate A, and B which will pull in C, D to active
 	LatentTestTransitionGFP(EGameFeatureTargetState::Active, PluginAURL);
 	LatentTestTransitionGFP(EGameFeatureTargetState::Active, PluginBURL);
 	LatentTestPluginState(EGameFeatureTargetState::Active, PluginCURL);
