Game Feature Plugins State Machine
- Allow skipping calls to InstallBundleManger::GetContentState() if content already up to date. This should remove most of these aysnc calls for GFPs loaded at startup or GFPs that have been predownloaded.

diff --git a/Engine/Plugins/Runtime/GameFeatures/Source/GameFeatures/Private/GameFeaturePluginStateMachine.cpp b/Engine/Plugins/Runtime/GameFeatures/Source/GameFeatures/Private/GameFeaturePluginStateMachine.cpp
--- a/Engine/Plugins/Runtime/GameFeatures/Source/GameFeatures/Private/GameFeaturePluginStateMachine.cpp
+++ b/Engine/Plugins/Runtime/GameFeatures/Source/GameFeatures/Private/GameFeaturePluginStateMachine.cpp
@@ -1832,55 +1832,37 @@ struct FGameFeaturePluginState_Downloading : public FGameFeaturePluginState
 		const FInstallBundlePluginProtocolOptions& Options = StateProperties.ProtocolOptions.GetSubtype<FInstallBundlePluginProtocolOptions>();
 
 		const bool bAddDependencies = false; // We already got all dependencies in the CheckingStatus state
-		if (InstallBundles.Num() > 1)
-		{
-			// Handle bDoNotDownload flag before doing any async ops
-			if (Options.bDoNotDownload)
-			{
-				TValueOrError<FInstallBundleCombinedInstallState, EInstallBundleResult> MaybeInstallState = 
-					BundleManager->GetInstallStateSynchronous(InstallBundles, bAddDependencies);
-				check(MaybeInstallState.HasValue());
-				FInstallBundleCombinedInstallState& InstallState = MaybeInstallState.GetValue();
 
-				// if not up to date, check to see if we allow downloading
-				if (!InstallState.GetAllBundlesHaveState(EInstallBundleInstallState::UpToDate))
-				{
-					Result = GetErrorResult(TEXT("GFPStateMachine.DownloadNotAllowed"));
-					bSuppressResultErrorLog = true; // Don't log an error if the user disallowed the download
-					UpdateStateMachineImmediate();
-					return;
-				}
-			}
+		TValueOrError<FInstallBundleCombinedInstallState, EInstallBundleResult> MaybeInstallState =
+			BundleManager->GetInstallStateSynchronous(InstallBundles, bAddDependencies);
+		check(MaybeInstallState.HasValue());
+		FInstallBundleCombinedInstallState& InstallState = MaybeInstallState.GetValue();
 
-			UGameFeaturesSubsystem::Get().OnGameFeatureDownloading(StateProperties.PluginName, StateProperties.PluginIdentifier);
-			EnsureAllowAsyncLoading();
+		const bool bAllUpToDate = InstallState.GetAllBundlesHaveState(EInstallBundleInstallState::UpToDate);
 
+		// Handle bDoNotDownload flag before doing any async ops
+		// if not up to date, check to see if we allow downloading
+		if (Options.bDoNotDownload && !bAllUpToDate)
+		{
+			Result = GetErrorResult(TEXT("GFPStateMachine.DownloadNotAllowed"));
+			bSuppressResultErrorLog = true; // Don't log an error if the user disallowed the download
+			UpdateStateMachineImmediate();
+			return;
+		}
+
+		UGameFeaturesSubsystem::Get().OnGameFeatureDownloading(StateProperties.PluginName, StateProperties.PluginIdentifier);
+		EnsureAllowAsyncLoading();
+
+		if (!bAllUpToDate && InstallBundles.Num() > 1)
+		{
 			GotContentStateHandle = BundleManager->GetContentState(InstallBundles, EInstallBundleGetContentStateFlags::None, bAddDependencies,
 				FInstallBundleGetContentStateDelegate::CreateRaw(this, &FGameFeaturePluginState_Downloading::OnGotContentState));
 		}
 		else
 		{
-			TValueOrError<FInstallBundleCombinedInstallState, EInstallBundleResult> MaybeInstallState = 
-				BundleManager->GetInstallStateSynchronous(InstallBundles, bAddDependencies);
-			check(MaybeInstallState.HasValue());
-			FInstallBundleCombinedInstallState& InstallState = MaybeInstallState.GetValue();
-
-			// Handle bDoNotDownload flag before doing any async ops
-			// if not up to date, check to see if we allow downloading
-			if (Options.bDoNotDownload && !InstallState.GetAllBundlesHaveState(EInstallBundleInstallState::UpToDate))
-			{
-				Result = GetErrorResult(TEXT("GFPStateMachine.DownloadNotAllowed"));
-				bSuppressResultErrorLog = true; // Don't log an error if the user disallowed the download
-				UpdateStateMachineImmediate();
-				return;
-			}
-
-			UGameFeaturesSubsystem::Get().OnGameFeatureDownloading(StateProperties.PluginName, StateProperties.PluginIdentifier);
-			EnsureAllowAsyncLoading();
-
-			// We usualy only have a use case of one bundle per GFP and we only care about relative weighting here, so we don't need 
-			// any of the other content state metadata. We can just assume the weight is 1.0 and not have to wait for the full 
-			// async call to get the rest of the metadata.
+			// Only if all bundles up to date or only one bundle:
+			// We only care about relative weighting here, so we don't need any of the other content state metadata.
+			// We can just assume the weight is 1.0 and not have to wait for the full async call to get the rest of the metadata.
 			FInstallBundleCombinedContentState HackContentState;
 			HackContentState.IndividualBundleStates.Reserve(InstallState.IndividualBundleStates.Num());
 			for (const TPair<FName, EInstallBundleInstallState>& Pair : InstallState.IndividualBundleStates)
