Add console command to print Game Feature Action usage

diff --git a/Engine/Plugins/Runtime/GameFeatures/Source/GameFeaturesEditor/Private/GameFeaturesEditorModule.cpp b/Engine/Plugins/Runtime/GameFeatures/Source/GameFeaturesEditor/Private/GameFeaturesEditorModule.cpp
--- a/Engine/Plugins/Runtime/GameFeatures/Source/GameFeaturesEditor/Private/GameFeaturesEditorModule.cpp
+++ b/Engine/Plugins/Runtime/GameFeatures/Source/GameFeaturesEditor/Private/GameFeaturesEditorModule.cpp
@@ -1,31 +1,214 @@
 // Copyright Epic Games, Inc. All Rights Reserved.
 
-#include "AssetRegistry/ARFilter.h"
-#include "Features/IModularFeatures.h"
-
-#include "GameFeatureData.h"
+#include "Engine/AssetManager.h"
+#include "Engine/AssetManagerSettings.h"
 #include "Features/EditorFeatures.h"
-
-#include "GameFeaturesSubsystem.h"
-#include "Interfaces/IPluginManager.h"
+#include "Features/IModularFeatures.h"
 #include "Features/IPluginsEditorFeature.h"
+#include "Framework/Notifications/NotificationManager.h"
+#include "GameFeatureData.h"
 #include "GameFeatureDataDetailsCustomization.h"
-#include "GameFeaturesEditorSettings.h"
-#include "GameFeaturesSubsystemSettings.h"
 #include "GameFeaturePluginMetadataCustomization.h"
 #include "GameFeaturePluginTemplate.h"
+#include "GameFeaturesEditorSettings.h"
+#include "GameFeaturesSubsystem.h"
+#include "GameFeaturesSubsystemSettings.h"
+#include "Interfaces/IPluginManager.h"
 #include "Logging/MessageLog.h"
 #include "Misc/App.h"
 #include "Modules/ModuleManager.h"
+#include "PropertyEditorModule.h"
 #include "SSettingsEditorCheckoutNotice.h"
-#include "Engine/AssetManagerSettings.h"
 #include "Widgets/Notifications/SNotificationList.h"
-#include "Framework/Notifications/NotificationManager.h"
-#include "Engine/AssetManager.h"
-#include "PropertyEditorModule.h"
 
 #define LOCTEXT_NAMESPACE "GameFeatures"
 
+static FAutoConsoleCommand CCmdGFPPrintActionUsageStats(
+	TEXT("GameFeaturePlugin.PrintActionUsage"),
+	TEXT("List all uses of Game Feature Actions")
+	TEXT("\nUsage: GameFeaturePlugin.PrintActionUsage [-Summary] [-Detailed [-Csv]]")
+	TEXT("\n-Summary - Print the number of Game Feature Plugins and the number of uses found for each Game Feature Action")
+	TEXT("\n-Details - Print the list of Game Feature Plugins and the list of plugins that use each Game Feature Action")
+	TEXT("\n-Csv     - Prints the output in CSV format. Only affects -Detailed output."),
+	FConsoleCommandWithArgsDelegate::CreateStatic([](const TArray<FString>& Args)
+	{
+		const UGameFeaturesSubsystem& GameFeatures = UGameFeaturesSubsystem::Get();
+
+		// Gather disabled GFPs
+		TArray<const FString*> DisabledPlugins;
+		TArray<const FString*> DisallowedPlugins;
+		{
+			DisabledPlugins.Reserve(32);
+			DisallowedPlugins.Reserve(32);
+
+			IPluginManager& PluginManager = IPluginManager::Get();
+			const UGameFeaturesSubsystemSettings& Settings = *GetDefault<UGameFeaturesSubsystemSettings>();
+
+			TArray<TSharedRef<IPlugin>> Plugins = PluginManager.GetDiscoveredPlugins();
+			for (const TSharedRef<IPlugin>& Plugin : Plugins)
+			{
+				const FString& DescriptorFileName = Plugin->GetDescriptorFileName();
+				if (Settings.IsValidGameFeaturePlugin(DescriptorFileName))
+				{
+					if (!Plugin->IsEnabled())
+					{
+						DisabledPlugins.Add(&Plugin->GetName());
+					}
+					else
+					{
+						FString PluginURL;
+						GameFeatures.GetBuiltInGameFeaturePluginURL(Plugin, PluginURL);
+						if (!GameFeatures.IsPluginAllowed(PluginURL))
+						{
+							DisallowedPlugins.Add(&Plugin->GetName());
+						}
+					}
+				}
+			}
+		}
+
+		// Gather actions
+		int32 RegisteredCount = 0;
+		TMap<const UClass*, TArray<FString>> ActionUses;
+		{
+			ActionUses.Reserve(128);
+
+			auto VisitGFP = [&RegisteredCount, &ActionUses](const UGameFeatureData* Data)
+			{
+				RegisteredCount++;
+
+				FString PluginName;
+				Data->GetPluginName(PluginName);
+
+				const TArray<UGameFeatureAction*>& Actions = Data->GetActions();
+				for (const UGameFeatureAction* Action : Actions)
+				{
+					const UClass* Class = Action->GetClass();
+					TArray<FString>& Plugins = ActionUses.FindOrAdd(Class);
+					Plugins.Add(PluginName);
+				}
+			};
+
+			GameFeatures.ForEachRegisteredGameFeature<UGameFeatureData>(VisitGFP);
+
+			ActionUses.ValueSort([](const TArray<FString>& A, const TArray<FString>& B) { return A.Num() > B.Num(); });
+		}
+
+		// Output
+		{
+			bool bDetails = Args.Contains(TEXT("-Details"));
+			bool bSummary = Args.Contains(TEXT("-Summary")) || !bDetails;
+			bool bCsv     = Args.Contains(TEXT("-Csv"));
+
+			TStringBuilder<1024> Msg;
+			Msg << TEXT("Game Feature Action usage\n");
+
+			if (bSummary)
+			{
+				if (!DisabledPlugins.IsEmpty())
+				{
+					Msg << TEXT("Disabled Game Feature Plugins (Action usage for these will not be displayed): ") << DisabledPlugins.Num() << TEXT("\n");
+				}
+
+				if (!DisallowedPlugins.IsEmpty())
+				{
+					Msg << TEXT("Disallowed built-in Game Feature Plugins (Action usage for these will not be displayed): ") << DisallowedPlugins.Num() << TEXT("\n");
+				}
+
+				Msg << TEXT("Registered Game Feature Plugins: ") << RegisteredCount << TEXT("\n");
+				if (!ActionUses.IsEmpty())
+				{
+					Msg << TEXT("Game Feature Action usage:\n{\n");
+					for (TPair<const UClass*, TArray<FString>>& Pair : ActionUses)
+					{
+						Msg << TEXT("\t") << Pair.Key->GetFName() << TEXT(": ") << Pair.Value.Num() << TEXT("\n");
+					}
+					Msg << TEXT("}\n");
+				}
+				else
+				{
+					Msg << TEXT("No Game Feature Actions found\n");
+				}
+			}
+
+			if (bDetails)
+			{
+				if (bCsv)
+				{
+					Msg << TEXT("Plugin Name,Action Class\n");
+
+					if (!DisabledPlugins.IsEmpty())
+					{
+						for (const FString* PluginName : DisabledPlugins)
+						{
+							Msg << *PluginName << TEXT(",Disabled\n");
+						}
+					}
+
+					if (!DisallowedPlugins.IsEmpty())
+					{
+						for (const FString* PluginName : DisallowedPlugins)
+						{
+							Msg << *PluginName << TEXT(",Disallowed\n");
+						}
+					}
+
+					for (TPair<const UClass*, TArray<FString>>& Pair : ActionUses)
+					{
+						FName ActionName = Pair.Key->GetFName();
+						for (const FString& PluginName : Pair.Value)
+						{
+							Msg << ActionName << TEXT(",") << PluginName << TEXT("\n");
+						}
+					}
+				}
+				else
+				{
+					if (!DisabledPlugins.IsEmpty())
+					{
+						Msg << TEXT("Disabled Game Feature Plugins (Action usage for these will not be displayed):\n{\n");
+						for (const FString* PluginName : DisabledPlugins)
+						{
+							Msg << TEXT("\t") << *PluginName << TEXT("\n");
+						}
+						Msg << TEXT("}\n");
+					}
+
+					if (!DisallowedPlugins.IsEmpty())
+					{
+						Msg << TEXT("Disallowed built-in Game Feature Plugins (Action usage for these will not be displayed):\n{\n");
+						for (const FString* PluginName : DisallowedPlugins)
+						{
+							Msg << TEXT("\t") << *PluginName << TEXT("\n");
+						}
+						Msg << TEXT("}\n");
+					}
+
+					Msg << TEXT("Game Feature Action usage:\n{\n");
+					for (TPair<const UClass*, TArray<FString>>& Pair : ActionUses)
+					{
+						FName ActionName = Pair.Key->GetFName();
+
+						Msg << TEXT("\t") << ActionName << TEXT(":\n\t{\n");
+						for (const FString& PluginName : Pair.Value)
+						{
+							Msg << TEXT("\t\t") << PluginName << TEXT("\n");
+						}
+						Msg << TEXT("\t}\n");
+					}
+					Msg << TEXT("}\n");
+				}
+			}
+
+			if (Msg.Len())
+			{
+				Msg.RemoveSuffix(1);
+			}
+			UE_LOG(LogGameFeatures, Display, TEXT("%s"), Msg.ToString());
+		}
+	})
+);
+
 class FGameFeaturesEditorModule : public FDefaultModuleImpl
 {
 	virtual void StartupModule() override
