Fixed reversed CSV output in CL 43545162

diff --git a/Engine/Plugins/Runtime/GameFeatures/Source/GameFeaturesEditor/Private/GameFeaturesEditorModule.cpp b/Engine/Plugins/Runtime/GameFeatures/Source/GameFeaturesEditor/Private/GameFeaturesEditorModule.cpp
--- a/Engine/Plugins/Runtime/GameFeatures/Source/GameFeaturesEditor/Private/GameFeaturesEditorModule.cpp
+++ b/Engine/Plugins/Runtime/GameFeatures/Source/GameFeaturesEditor/Private/GameFeaturesEditorModule.cpp
@@ -158,7 +158,7 @@ static FAutoConsoleCommand CCmdGFPPrintActionUsageStats(
 						FName ActionName = Pair.Key->GetFName();
 						for (const FString& PluginName : Pair.Value)
 						{
-							Msg << ActionName << TEXT(",") << PluginName << TEXT("\n");
+							Msg << PluginName << TEXT(",") << ActionName << TEXT("\n");
 						}
 					}
 				}
