FGameFeaturePluginState_Downloading - Fix IoStoreOnDemand requests getting queued if install bundles fail. The install bundles are required to mount the ondemand containers.

diff --git a/Engine/Plugins/Runtime/GameFeatures/Source/GameFeatures/Private/GameFeaturePluginStateMachine.cpp b/Engine/Plugins/Runtime/GameFeatures/Source/GameFeatures/Private/GameFeaturePluginStateMachine.cpp
--- a/Engine/Plugins/Runtime/GameFeatures/Source/GameFeatures/Private/GameFeaturePluginStateMachine.cpp
+++ b/Engine/Plugins/Runtime/GameFeatures/Source/GameFeatures/Private/GameFeaturePluginStateMachine.cpp
@@ -1691,6 +1691,12 @@ struct FGameFeaturePluginState_Downloading : public FGameFeaturePluginState
 
 		OnUpdateProgress(0.0f);
 
+		if (Result.HasError())
+		{
+			UpdateStateMachineImmediate();
+			return;
+		}
+
 		InstallIoStoreOnDemandContent();
 	}
 
