[Engine] Add ReleaseGameFeaturePlugins for multiple plugins at once.
Unloading plugins batches GC and GameplayTags tree changes together which improves performance of unloading operation

#tests local build

diff --git a/Engine/Plugins/Runtime/GameFeatures/Source/GameFeatures/GameFeatures.Build.cs b/Engine/Plugins/Runtime/GameFeatures/Source/GameFeatures/GameFeatures.Build.cs
--- a/Engine/Plugins/Runtime/GameFeatures/Source/GameFeatures/GameFeatures.Build.cs
+++ b/Engine/Plugins/Runtime/GameFeatures/Source/GameFeatures/GameFeatures.Build.cs
@@ -4,16 +4,18 @@ namespace UnrealBuildTool.Rules
 {
 	public class GameFeatures : ModuleRules
 	{
-        public GameFeatures(ReadOnlyTargetRules Target) : base(Target)
+		public GameFeatures(ReadOnlyTargetRules Target) : base(Target)
 		{
+			bTreatAsEngineModule = true;
+
 			PublicDependencyModuleNames.AddRange(
 				new string[]
-                {
-                    "Core",
-                    "CoreUObject",
+				{
+					"Core",
+					"CoreUObject",
 					"DeveloperSettings",
 					"Engine",
-                    "ModularGameplay",
+					"ModularGameplay",
 					"DataRegistry"
 				}
 			);
diff --git a/Engine/Plugins/Runtime/GameFeatures/Source/GameFeatures/Private/GameFeaturesSubsystem.cpp b/Engine/Plugins/Runtime/GameFeatures/Source/GameFeatures/Private/GameFeaturesSubsystem.cpp
--- a/Engine/Plugins/Runtime/GameFeatures/Source/GameFeatures/Private/GameFeaturesSubsystem.cpp
+++ b/Engine/Plugins/Runtime/GameFeatures/Source/GameFeatures/Private/GameFeaturesSubsystem.cpp
@@ -1851,6 +1851,66 @@ void UGameFeaturesSubsystem::ReleaseGameFeaturePlugin(const FString& PluginURL,
 	}
 }
 
+void UGameFeaturesSubsystem::ReleaseGameFeaturePlugin(TConstArrayView<FString> PluginURLs, const FMultipleGameFeaturePluginsReleased& CompleteDelegate)
+{
+	struct FReleaseContext
+	{
+		TMap<FString, UE::GameFeatures::FResult> Results;
+		FMultipleGameFeaturePluginsReleased CompleteDelegate;
+		int32 NumPluginsReleased = 0;
+		bool bDeferGCAfterUnload = false;
+		bool bPushedTagsBroadcast = false;
+
+		FReleaseContext()
+		{
+			if (!IsEngineExitRequested())
+			{
+				UGameplayTagsManager::Get().PushDeferOnGameplayTagTreeChangedBroadcast();
+				IPluginManager::Get().SuppressPluginUnloadGC();
+				bDeferGCAfterUnload = true;
+				bPushedTagsBroadcast = true;
+			}
+			else if(UGameplayTagsManager* TagsManager = UGameplayTagsManager::GetIfAllocated())
+			{
+				TagsManager->PushDeferOnGameplayTagTreeChangedBroadcast();
+				bPushedTagsBroadcast = true;
+			}
+		}
+
+		~FReleaseContext()
+		{
+			if (bDeferGCAfterUnload)
+			{
+				IPluginManager::Get().ResumePluginUnloadGC();
+			}
+			if (bPushedTagsBroadcast)
+			{
+				UGameplayTagsManager::Get().PopDeferOnGameplayTagTreeChangedBroadcast();
+			}
+
+			CompleteDelegate.ExecuteIfBound(Results);
+		}
+	};
+	TSharedRef<FReleaseContext> LoadContext = MakeShared<FReleaseContext>();
+	LoadContext->CompleteDelegate = CompleteDelegate;
+
+	LoadContext->Results.Reserve(PluginURLs.Num());
+	for (const FString& PluginURL : PluginURLs)
+	{
+		LoadContext->Results.Add(PluginURL, MakeError("Pending"));
+	}
+
+	for (const FString& PluginURL : PluginURLs)
+	{
+		ReleaseGameFeaturePlugin(PluginURL, FGameFeaturePluginReleaseComplete::CreateLambda([LoadContext, PluginURL](const UE::GameFeatures::FResult& Result)
+		{
+			LoadContext->Results[PluginURL] = Result;
+			++LoadContext->NumPluginsReleased;
+			UE_LOG(LogGameFeatures, VeryVerbose, TEXT("Finished Releasing %i/%i GFPs"), LoadContext->NumPluginsReleased, LoadContext->Results.Num());
+		}));
+	}
+}
+
 void UGameFeaturesSubsystem::UninstallGameFeaturePlugin(const FString& PluginURL, const FGameFeaturePluginUninstallComplete& CompleteDelegate /*= FGameFeaturePluginUninstallComplete()*/)
 {
 	UninstallGameFeaturePlugin(PluginURL, FGameFeatureProtocolOptions(), CompleteDelegate);
diff --git a/Engine/Plugins/Runtime/GameFeatures/Source/GameFeatures/Public/GameFeaturesSubsystem.h b/Engine/Plugins/Runtime/GameFeatures/Source/GameFeatures/Public/GameFeaturesSubsystem.h
--- a/Engine/Plugins/Runtime/GameFeatures/Source/GameFeatures/Public/GameFeaturesSubsystem.h
+++ b/Engine/Plugins/Runtime/GameFeatures/Source/GameFeatures/Public/GameFeaturesSubsystem.h
@@ -217,6 +217,7 @@ using FMultipleGameFeaturePluginChangeStateComplete = TDelegate<void(const TMap<
 
 using FBuiltInGameFeaturePluginsLoaded = FMultipleGameFeaturePluginChangeStateComplete;
 using FMultipleGameFeaturePluginsLoaded = FMultipleGameFeaturePluginChangeStateComplete;
+using FMultipleGameFeaturePluginsReleased = FMultipleGameFeaturePluginChangeStateComplete;
 using FMultipleGameFeaturePluginsTerminated = FMultipleGameFeaturePluginChangeStateComplete;
 
 enum class EBuiltInAutoState : uint8
@@ -588,6 +589,7 @@ public:
 	/** Releases any game data stored for this GameFeaturePlugin. Does not uninstall data and it will remain on disk. */
 	UE_API void ReleaseGameFeaturePlugin(const FString& PluginURL);
 	UE_API void ReleaseGameFeaturePlugin(const FString& PluginURL, const FGameFeaturePluginReleaseComplete& CompleteDelegate);
+	UE_API void ReleaseGameFeaturePlugin(TConstArrayView<FString> PluginURLs, const FMultipleGameFeaturePluginsReleased& CompleteDelegate);
 
 	/** Uninstalls any game data stored for this GameFeaturePlugin and terminates the GameFeaturePlugin.
 		If the given PluginURL is not found this will create a GameFeaturePlugin first and attempt to run it through the uninstall flow.
