Bugfix: Changes to game feature plugin details made via "Edit Plugin" button appear to be ignored

The .uplugin was actually being updated correctly, but re-launching the in-editor plugin details UI would show the stale, pre-edited value. Caused by our in-memory GameFeaturePlugin details cache getting out of sync with the file on disk.

diff --git a/Engine/Plugins/Runtime/GameFeatures/Source/GameFeatures/Private/GameFeaturesSubsystem.cpp b/Engine/Plugins/Runtime/GameFeatures/Source/GameFeatures/Private/GameFeaturesSubsystem.cpp
--- a/Engine/Plugins/Runtime/GameFeatures/Source/GameFeatures/Private/GameFeaturesSubsystem.cpp
+++ b/Engine/Plugins/Runtime/GameFeatures/Source/GameFeatures/Private/GameFeaturesSubsystem.cpp
@@ -500,7 +500,8 @@ void UGameFeaturesSubsystem::Initialize(FSubsystemCollectionBase& Collection)
 #endif
 
 	GetExplanationForUnavailablePackageDelegateHandle = IPluginManager::Get().GetExplanationForUnavailablePackageWithPluginInfoDelegate().AddUObject(this, &UGameFeaturesSubsystem::GetExplanationForUnavailablePackage);
-	OnPluginUnmountedDelegateHandle = IPluginManager::Get().OnPluginUnmounted().AddWeakLambda(this, [this] (const IPlugin& Plugin) { PruneCachedGameFeaturePluginDetails(/* doesnt use URL */ FString(""), Plugin.GetDescriptorFileName()); });
+	IPluginManager::Get().OnPluginUnmounted().AddWeakLambda(this, [this] (const IPlugin& Plugin) { PruneCachedGameFeaturePluginDetails(/* doesnt use URL */ FString(""), Plugin.GetDescriptorFileName()); });
+	IPluginManager::Get().OnPluginEdited().AddWeakLambda(this, [this](IPlugin& Plugin) { PruneCachedGameFeaturePluginDetails(/* doesnt use URL */ FString(""), Plugin.GetDescriptorFileName()); });
 }
 
 void UGameFeaturesSubsystem::Deinitialize()
@@ -508,7 +509,6 @@ void UGameFeaturesSubsystem::Deinitialize()
 	UE_LOG(LogGameFeatures, Log, TEXT("Shutting down game features subsystem"));
 	IPluginManager::Get().GetExplanationForUnavailablePackageWithPluginInfoDelegate().Remove(GetExplanationForUnavailablePackageDelegateHandle);
 	GetExplanationForUnavailablePackageDelegateHandle.Reset();
-	OnPluginUnmountedDelegateHandle.Reset();
 
 	if ((GameSpecificPolicies != nullptr) && bInitializedPolicyManager)
 	{
diff --git a/Engine/Plugins/Runtime/GameFeatures/Source/GameFeatures/Public/GameFeaturesSubsystem.h b/Engine/Plugins/Runtime/GameFeatures/Source/GameFeatures/Public/GameFeaturesSubsystem.h
--- a/Engine/Plugins/Runtime/GameFeatures/Source/GameFeatures/Public/GameFeaturesSubsystem.h
+++ b/Engine/Plugins/Runtime/GameFeatures/Source/GameFeatures/Public/GameFeaturesSubsystem.h
@@ -918,7 +918,6 @@ private:
 	TSet<FString> DebugStateChangedForPlugins;
 #endif
 	FDelegateHandle GetExplanationForUnavailablePackageDelegateHandle;
-	FDelegateHandle OnPluginUnmountedDelegateHandle;
 
 	bool bInitializedPolicyManager = false;
 };
