[GFP] Expose GetPluginNameFromPluginURL in GameFeatureSubsystem

#rnx

diff --git a/Engine/Plugins/Runtime/GameFeatures/Source/GameFeatures/Private/GameFeaturesSubsystem.cpp b/Engine/Plugins/Runtime/GameFeatures/Source/GameFeatures/Private/GameFeaturesSubsystem.cpp
--- a/Engine/Plugins/Runtime/GameFeatures/Source/GameFeatures/Private/GameFeaturesSubsystem.cpp
+++ b/Engine/Plugins/Runtime/GameFeatures/Source/GameFeatures/Private/GameFeaturesSubsystem.cpp
@@ -2500,6 +2500,17 @@ FString UGameFeaturesSubsystem::GetPluginFilenameFromPluginURL(const FString& Pl
 	return PluginFilename;
 }
 
+FString UGameFeaturesSubsystem::GetPluginNameFromPluginURL(const FString& PluginURL) const
+{
+	const UGameFeaturePluginStateMachine* GFSM = FindGameFeaturePluginStateMachine(PluginURL);
+	if (GFSM == nullptr)
+	{
+		UE_LOG(LogGameFeatures, Error, TEXT("UGameFeaturesSubsystem could not get the plugin name from the plugin URL. URL:%s "), *PluginURL);
+		return FString();
+	}
+	return GFSM->GetPluginName();
+}
+
 void UGameFeaturesSubsystem::FixPluginPackagePath(FString& PathToFix, const FString& PluginRootPath, bool bMakeRelativeToPluginRoot)
 {
 	if (bMakeRelativeToPluginRoot)
diff --git a/Engine/Plugins/Runtime/GameFeatures/Source/GameFeatures/Public/GameFeaturesSubsystem.h b/Engine/Plugins/Runtime/GameFeatures/Source/GameFeatures/Public/GameFeaturesSubsystem.h
--- a/Engine/Plugins/Runtime/GameFeatures/Source/GameFeatures/Public/GameFeaturesSubsystem.h
+++ b/Engine/Plugins/Runtime/GameFeatures/Source/GameFeatures/Public/GameFeaturesSubsystem.h
@@ -621,6 +621,9 @@ public:
 	/** Get the plugin path from the plugin URL */
 	UE_API FString GetPluginFilenameFromPluginURL(const FString& PluginURL) const;
 
+	/** Get the plugin name from the plugin URL */
+	UE_API FString GetPluginNameFromPluginURL(const FString& PluginURL) const;
+
 	/** Fixes a package path/directory to either be relative to plugin root or not. Paths relative to different roots will not be modified */
 	static UE_API void FixPluginPackagePath(FString& PathToFix, const FString& PluginRootPath, bool bMakeRelativeToPluginRoot);
 
