Check all GameFeaturePlugin State machines and only release bundles which are not in use when finishing a predownload.

diff --git a/Engine/Plugins/Runtime/GameFeatures/Source/GameFeatures/Private/GameFeaturesSubsystem.cpp b/Engine/Plugins/Runtime/GameFeatures/Source/GameFeatures/Private/GameFeaturesSubsystem.cpp
--- a/Engine/Plugins/Runtime/GameFeatures/Source/GameFeatures/Private/GameFeaturesSubsystem.cpp
+++ b/Engine/Plugins/Runtime/GameFeatures/Source/GameFeatures/Private/GameFeaturesSubsystem.cpp
@@ -3061,13 +3061,28 @@ struct FGameFeaturePluginPredownloadContext : public FGameFeaturePluginPredownlo
 		TSharedPtr<IInstallBundleManager> BundleManager = IInstallBundleManager::GetPlatformInstallBundleManager();
 
 		TArray<FName> ReleaseList;
-		TArray<FName> KeepList;
-		for (const TPair<FGameFeaturePluginIdentifier, FInstallBundlePluginProtocolMetaData>& Pair : GFPs)
+		TSet<FName> KeepList;
+		for (const TPair<FString, TObjectPtr<UGameFeaturePluginStateMachine>>& Pair : GFPSubSys.GameFeaturePluginStateMachines)
 		{
-			UGameFeaturePluginStateMachine* Machine = GFPSubSys.FindGameFeaturePluginStateMachine(Pair.Key);
+			UGameFeaturePluginStateMachine* Machine = Pair.Value;
 			if (Machine &&
 				Machine->GetCurrentState() > EGameFeaturePluginState::StatusKnown &&
 				Machine->GetCurrentState() != EGameFeaturePluginState::Releasing)
+			{
+				const FGameFeatureProtocolMetadata& ProtocolMetaData = Machine->GetProtocolMetadata();
+				if (ProtocolMetaData.HasSubtype<FInstallBundlePluginProtocolMetaData>())
+				{
+					const FInstallBundlePluginProtocolMetaData& ProtocolData = ProtocolMetaData.GetSubtype<FInstallBundlePluginProtocolMetaData>();
+					KeepList.Append(ProtocolData.InstallBundles);
+				}
+			}
+		}
+		for (const TPair<FGameFeaturePluginIdentifier, FInstallBundlePluginProtocolMetaData>& Pair : GFPs)
+		{
+			UGameFeaturePluginStateMachine* Machine = GFPSubSys.FindGameFeaturePluginStateMachine(Pair.Key);
+			if (Machine
+				&& Machine->GetCurrentState() > EGameFeaturePluginState::StatusKnown
+				&& Machine->GetCurrentState() != EGameFeaturePluginState::Releasing)
 			{
 				// A machine is using the bundles, don't release
 				KeepList.Append(Pair.Value.InstallBundles);
@@ -3080,7 +3095,7 @@ struct FGameFeaturePluginPredownloadContext : public FGameFeaturePluginPredownlo
 			ReleaseList.Append(Pair.Value.InstallBundles);
 		}
 
-		BundleManager->RequestReleaseContent(ReleaseList, EInstallBundleReleaseRequestFlags::None, KeepList);
+		BundleManager->RequestReleaseContent(ReleaseList, EInstallBundleReleaseRequestFlags::None, KeepList.Array());
 	}
 };
 
