IAD support for GFP predownload.

diff --git a/Engine/Plugins/Runtime/GameFeatures/Source/GameFeatures/Private/GameFeaturesSubsystem.cpp b/Engine/Plugins/Runtime/GameFeatures/Source/GameFeatures/Private/GameFeaturesSubsystem.cpp
--- a/Engine/Plugins/Runtime/GameFeatures/Source/GameFeatures/Private/GameFeaturesSubsystem.cpp
+++ b/Engine/Plugins/Runtime/GameFeatures/Source/GameFeatures/Private/GameFeaturesSubsystem.cpp
@@ -2,6 +2,7 @@
 
 #include "GameFeaturesSubsystem.h"
 #include "Algo/Accumulate.h"
+#include "Algo/AllOf.h"
 #include "Algo/AnyOf.h"
 #include "Algo/TopologicalSort.h"
 #include "Algo/Unique.h"
@@ -16,6 +17,7 @@
 #include "GameFeatureStateChangeObserver.h"
 #include "GameplayTagsManager.h"
 #include "Interfaces/IPluginManager.h"
+#include "IO/IoStoreOnDemand.h"
 #include "Logging/StructuredLog.h"
 #include "Misc/App.h"
 #include "Misc/FileHelper.h"
@@ -1693,6 +1695,17 @@ bool UGameFeaturesSubsystem::ShouldUpdatePluginProtocolOptions(const UGameFeatur
 	return true;
 }
 
+namespace UE::GameFeatures
+{
+	static float CombineInstallProgress(float InstallProgress, float AssetDependencyProgress)
+	{
+		// Assumuption that most of the progress will be from asset dependencies in this case
+		// For this to be more accurate we'd need to figure out the actual sizes during 
+		// EGameFeaturePluginState::CheckingStatus but this is most likely good enough
+		return 0.2f * InstallProgress + 0.8f * AssetDependencyProgress;
+	}
+}
+
 bool UGameFeaturesSubsystem::GetGameFeaturePluginInstallPercent(const FString& PluginURL, float& Install_Percent) const
 {
 	if (const UGameFeaturePluginStateMachine* StateMachine = FindGameFeaturePluginStateMachine(PluginURL))
@@ -1727,10 +1740,7 @@ bool UGameFeaturesSubsystem::GetGameFeaturePluginInstallPercent(const FString& P
 				AssetDependencyProgress = 1.0f;
 			}
 
-			// Assumuption that most of the progress will be from asset dependencies in this case
-			// For this to be more accurate we'd need to figure out the actual sizes during 
-			// EGameFeaturePluginState::CheckingStatus but this is most likely good enough
-			Install_Percent = 0.2f * InstallProgress + 0.8f * AssetDependencyProgress;
+			Install_Percent = UE::GameFeatures::CombineInstallProgress(InstallProgress, AssetDependencyProgress);
 			return true;
 		}
 	}
@@ -2769,14 +2779,38 @@ void UGameFeaturesSubsystem::PruneCachedGameFeaturePluginDetails(const FString&
 
 struct FGameFeaturePluginPredownloadContext : public FGameFeaturePluginPredownloadHandle
 {
-	const FStringView PredownloadErrorNamespace = TEXTVIEW("GameFeaturePlugin.Predownload.");
-	TMap<FGameFeaturePluginIdentifier, FInstallBundlePluginProtocolMetaData> GFPs;
+	struct FGFPData
+	{
+		FInstallBundlePluginProtocolMetaData ProtocolMetadata;
+		TArray<EStreamingAssetInstallMode> BundleInstallModes;
+	};
+
+	static constexpr const FStringView PredownloadErrorNamespace = TEXTVIEW("GameFeaturePlugin.Predownload.");
+
+	UE::GameFeatures::FResult Result = MakeValue();
+
+	FTSTicker::FDelegateHandle TickHandle;
+
+	TMap<FGameFeaturePluginIdentifier, FGFPData> GFPs;
+
 	TArray<FName> PendingBundleDownloads;
-	TOptional<FInstallBundleCombinedProgressTracker> ProgressTracker;
+
+	FInstallBundleCombinedProgressTracker ProgressTracker{ false /*bAutoTick*/};
+
 	TUniqueFunction<void(const UE::GameFeatures::FResult&)> OnComplete;
 	TUniqueFunction<void(float)> OnProgress;
-	UE::GameFeatures::FResult Result = MakeValue();
+
+	TArray<UE::IoStore::FOnDemandContentHandle> IoStoreContentHandles;
+	TArray<UE::IoStore::FOnDemandInstallRequest> IoStoreInstallRequests;
+	TArray<UE::IoStore::FOnDemandInstallProgress> IoStoreProgress;
+	FIoStatus IoStoreStatus;
+	int32 IoStorePendingInstalls = 0;
+
 	float Progress = 0.0f;
+
+	UE::FMutex IoStoreMutex;
+
+	bool bHasAssetDependencies = false;
 	bool bIsComplete = false;
 	bool bCanceled = false;
 
@@ -2812,11 +2846,24 @@ struct FGameFeaturePluginPredownloadContext : public FGameFeaturePluginPredownlo
 				BundleManager->CancelUpdateContent(PendingBundleDownloads);
 			}
 		}
+
+		UE::TUniqueLock Lock(IoStoreMutex);
+		CancelIoStoreRequests();
+	}
+
+	bool Tick(float /*dt*/)
+	{
+		UpdateProgress();
+
+		return true;
 	}
 
 	void Cleanup()
 	{
-		ProgressTracker.Reset();
+		if (TickHandle.IsValid())
+		{
+			FTSTicker::GetCoreTicker().RemoveTicker(TickHandle);
+		}		
 		IInstallBundleManager::InstallBundleCompleteDelegate.RemoveAll(this);
 		IInstallBundleManager::PausedBundleDelegate.RemoveAll(this);
 	}
@@ -2894,12 +2941,27 @@ struct FGameFeaturePluginPredownloadContext : public FGameFeaturePluginPredownlo
 			return;
 		}
 
+		TSharedPtr<IInstallBundleManager> BundleManager = IInstallBundleManager::GetPlatformInstallBundleManager();
+		if (!BundleManager)
+		{
+			UE_LOGFMT(LogGameFeatures, Error, "GFP Predownload failed, no Install Bundle Manager found.");
+			UE::GameFeatures::FResult ErrorResult = MakeError(FString::Printf(TEXT("%.*s%s"),
+				PredownloadErrorNamespace.Len(), PredownloadErrorNamespace.GetData(),
+				TEXT("BundleManager_Null")));
+			SetComplete(MoveTemp(ErrorResult));
+			return;
+		}
+
 		UGameFeaturesSubsystem& GFPSubSys = UGameFeaturesSubsystem::Get();
 
+		bool bAllBundlesUpToDate = true; 
 		TArray<FName> BundlesToInstall;
-		for (const TPair<FGameFeaturePluginIdentifier, FInstallBundlePluginProtocolMetaData>& Pair : GFPs)
+		for (TPair<FGameFeaturePluginIdentifier, FGFPData>& Pair : GFPs)
 		{
-			UGameFeaturePluginStateMachine* Machine = GFPSubSys.FindGameFeaturePluginStateMachine(Pair.Key);
+			const FGameFeaturePluginIdentifier& GFPIdentifier = Pair.Key;
+			FGFPData& GFPData = Pair.Value;
+
+			UGameFeaturePluginStateMachine* Machine = GFPSubSys.FindGameFeaturePluginStateMachine(GFPIdentifier);
 			if (Machine && Machine->GetDestination() < EGameFeaturePluginState::Installed)
 			{
 				// Existing machine exists and wants to be uninstalled, can't precache
@@ -2911,33 +2973,84 @@ struct FGameFeaturePluginPredownloadContext : public FGameFeaturePluginPredownlo
 				return;
 			}
 
-			BundlesToInstall.Append(Pair.Value.InstallBundles);
+			GFPSubSys.OnGameFeaturePredownloading(FString(GFPIdentifier.GetPluginName()), GFPIdentifier);
 
-			GFPSubSys.OnGameFeaturePredownloading(FString(Pair.Key.GetPluginName()), Pair.Key);
-		}
+			const bool bAddDependencies = true;
+			TValueOrError<FInstallBundleCombinedInstallState, EInstallBundleResult> MaybeInstallState = 
+				BundleManager->GetInstallStateSynchronous(GFPData.ProtocolMetadata.InstallBundles, bAddDependencies);
+			if (MaybeInstallState.HasError())
+			{
+				UE_LOGFMT(LogGameFeatures, Error, "GFP Predownload failed, failed to get install state for {GFP}", ("GFP", GFPIdentifier.GetPluginName()));
+				UE::GameFeatures::FResult ErrorResult = MakeError(FString::Printf(TEXT("%.*s%s"),
+					PredownloadErrorNamespace.Len(), PredownloadErrorNamespace.GetData(),
+					LexToString(MaybeInstallState.GetError())));
+				SetComplete(MoveTemp(ErrorResult));
+				return;
+			}
 
-		TSharedPtr<IInstallBundleManager> BundleManager = IInstallBundleManager::GetPlatformInstallBundleManager();
-		if (!BundleManager)
-		{
-			UE_LOGFMT(LogGameFeatures, Error, "GFP Predownload failed, no Install Bundle Manager found.");
-			UE::GameFeatures::FResult ErrorResult = MakeError(FString::Printf(TEXT("%.*s%s"),
-				PredownloadErrorNamespace.Len(), PredownloadErrorNamespace.GetData(),
-				TEXT("BundleManager_Null")));
-			SetComplete(MoveTemp(ErrorResult));
-			return;
+			const FInstallBundleCombinedInstallState& InstallState = MaybeInstallState.GetValue();
+			const bool bIsAvailable = Algo::AllOf(GFPData.ProtocolMetadata.InstallBundles,
+				[&InstallState](FName BundleName) { return InstallState.IndividualBundleStates.Contains(BundleName); });
+
+			if (!bIsAvailable)
+			{
+				UE_LOGFMT(LogGameFeatures, Error, "GFP Predownload failed, unvailable {GFP}", ("GFP", GFPIdentifier.GetPluginName()));
+				UE::GameFeatures::FResult ErrorResult = MakeError(FString::Printf(TEXT("%.*s%s"),
+					PredownloadErrorNamespace.Len(), PredownloadErrorNamespace.GetData(),
+					TEXT("Plugin_Unavailable")));
+				SetComplete(MoveTemp(ErrorResult));
+				return;
+			}
+
+			GFPData.ProtocolMetadata.InstallBundlesWithAssetDependencies = InstallState.BundlesWithIoStoreOnDemand.Array();
+			if (!GFPData.ProtocolMetadata.InstallBundlesWithAssetDependencies.IsEmpty())
+			{
+				TValueOrError<TArray<EStreamingAssetInstallMode>, FString> MaybeInstallModes =
+					UGameFeaturesSubsystem::Get().GetPolicy().GetStreamingAssetInstallModes(
+						GFPIdentifier.GetFullPluginURL(), GFPData.ProtocolMetadata.InstallBundlesWithAssetDependencies);
+
+				if (MaybeInstallModes.HasError())
+				{
+					UE_LOGFMT(LogGameFeatures, Error, "GFP Predownload failed, no install modes for {GFP}", ("GFP", GFPIdentifier.GetPluginName()));
+					UE::GameFeatures::FResult ErrorResult = MakeError(FString::Printf(TEXT("%.*s%s"),
+						PredownloadErrorNamespace.Len(), PredownloadErrorNamespace.GetData(),
+						*MaybeInstallModes.GetError()));
+					SetComplete(MoveTemp(ErrorResult));
+					return;
+				}
+
+				GFPData.BundleInstallModes = MaybeInstallModes.StealValue();
+				check(GFPData.BundleInstallModes.Num() == GFPData.ProtocolMetadata.InstallBundlesWithAssetDependencies.Num());
+			}
+
+			bHasAssetDependencies = bHasAssetDependencies || !InstallState.BundlesWithIoStoreOnDemand.IsEmpty();
+
+			bAllBundlesUpToDate = bAllBundlesUpToDate && 
+				InstallState.BundlesWithIoStoreOnDemand.IsEmpty() && // For now, just assume that any IAD data is not up to date
+				InstallState.GetAllBundlesHaveState(EInstallBundleInstallState::UpToDate);
+
+			// Update metadata with fully expanded dependency list. This can only be done after all bundles are known to be available,
+			// otherwise unavailable bundles in the URL could be stripped from the list.
+			GFPData.ProtocolMetadata.InstallBundles.Empty(InstallState.IndividualBundleStates.Num());
+			InstallState.IndividualBundleStates.GetKeys(GFPData.ProtocolMetadata.InstallBundles);
+
+			// Its ok to have duplicates in this list
+			BundlesToInstall.Append(GFPData.ProtocolMetadata.InstallBundles);
 		}
 
 		// Early out if everything is up to date already. This helps avoid enquing UI dialogs for content that doesn't actually need to be downloaded
-		TValueOrError<FInstallBundleCombinedInstallState, EInstallBundleResult> MaybeInstallState = BundleManager->GetInstallStateSynchronous(BundlesToInstall, false);
-		if (MaybeInstallState.HasValue() && MaybeInstallState.GetValue().GetAllBundlesHaveState(EInstallBundleInstallState::UpToDate))
+		if (bAllBundlesUpToDate)
 		{
 			SetComplete(MakeValue());
 			return;
 		}
 
 		BundleManager->GetContentState(BundlesToInstall, EInstallBundleGetContentStateFlags::None, false,
-			FInstallBundleGetContentStateDelegate::CreateLambda([Context = SharedThis(this)](FInstallBundleCombinedContentState BundleContentState)
-			{ Context->OnGotContentState(MoveTemp(BundleContentState)); }));
+			FInstallBundleGetContentStateDelegate::CreateLambda(
+				[Context = SharedThis(this)](FInstallBundleCombinedContentState BundleContentState)
+				{ Context->OnGotContentState(MoveTemp(BundleContentState)); }
+			)
+		);
 	}
 
 	void OnGotContentState(FInstallBundleCombinedContentState BundleContentState)
@@ -2946,7 +3059,7 @@ struct FGameFeaturePluginPredownloadContext : public FGameFeaturePluginPredownlo
 
 		if (BundleContentState.GetAllBundlesHaveState(EInstallBundleInstallState::UpToDate))
 		{
-			SetComplete(MakeValue());
+			OnAllInstallBundlesCompleted();
 			return;
 		}
 
@@ -2957,9 +3070,10 @@ struct FGameFeaturePluginPredownloadContext : public FGameFeaturePluginPredownlo
 		}
 
 		TArray<FName> BundlesToInstall;
-		for (const TPair<FGameFeaturePluginIdentifier, FInstallBundlePluginProtocolMetaData>& Pair : GFPs)
+		for (const TPair<FGameFeaturePluginIdentifier, FGFPData>& Pair : GFPs)
 		{
-			BundlesToInstall.Append(Pair.Value.InstallBundles);
+			const FGFPData& GFPData = Pair.Value;
+			BundlesToInstall.Append(GFPData.ProtocolMetadata.InstallBundles);
 		}
 
 		EInstallBundleRequestFlags InstallFlags = EInstallBundleRequestFlags::Defaults | EInstallBundleRequestFlags::SkipMount;
@@ -2978,21 +3092,16 @@ struct FGameFeaturePluginPredownloadContext : public FGameFeaturePluginPredownlo
 		FInstallBundleRequestInfo RequestInfo = MaybeRequestInfo.StealValue();
 		if (RequestInfo.BundlesEnqueued.Num() == 0)
 		{
-			SetComplete(MakeValue());
+			OnAllInstallBundlesCompleted();
 			return;
 		}
 
 		PendingBundleDownloads = MoveTemp(RequestInfo.BundlesEnqueued);
 
-		ProgressTracker.Emplace(true, [this](const FInstallBundleCombinedProgressTracker::FCombinedProgress& InProgress)
-		{
-			Progress = InProgress.ProgressPercent;
-			if (OnProgress)
-			{
-				OnProgress(InProgress.ProgressPercent);
-			}
-		});
-		ProgressTracker->SetBundlesToTrackFromContentState(BundleContentState, PendingBundleDownloads);
+		ProgressTracker.SetBundlesToTrackFromContentState(BundleContentState, PendingBundleDownloads);
+
+		// Start ticking
+		TickHandle = FTSTicker::GetCoreTicker().AddTicker(FTickerDelegate::CreateSP(this, &FGameFeaturePluginPredownloadContext::Tick));
 
 		IInstallBundleManager::InstallBundleCompleteDelegate.AddLambda(
 			[Context = SharedThis(this)](FInstallBundleRequestResultInfo BundleResult)
@@ -3008,7 +3117,7 @@ struct FGameFeaturePluginPredownloadContext : public FGameFeaturePluginPredownlo
 			return;
 		}
 
-		PendingBundleDownloads.Remove(BundleResult.BundleName);
+		PendingBundleDownloads.RemoveSwap(BundleResult.BundleName);
 
 		if (Result.HasValue() && BundleResult.Result != EInstallBundleResult::OK)
 		{
@@ -3041,22 +3150,202 @@ struct FGameFeaturePluginPredownloadContext : public FGameFeaturePluginPredownlo
 			return;
 		}
 
-		// Delay call to ReleaseBundlesIfPossible. We don't want to release them from within the complete callback.
+		// Delay call. We don't want to possibly release bundles from within the complete callback.
 		FTSTicker::GetCoreTicker().AddTicker(
 			FTickerDelegate::CreateLambda([Context = SharedThis(this)](float)
 			{
-				// Done
-				Context->SetComplete();
-
-				Context->Cleanup();
+				Context->OnAllInstallBundlesCompleted();
 				return false;
 			})
 		);
 	}
 
+	void OnAllInstallBundlesCompleted()
+	{
+		if (Result.HasError() || !bHasAssetDependencies)
+		{
+			SetComplete();
+			Cleanup();
+		}
+
+		UE::IoStore::IOnDemandIoStore* IoStore = UE::IoStore::TryGetOnDemandIoStore();
+		if (!IoStore)
+		{
+			UE_LOGFMT(LogGameFeatures, Error, "GFP Predownload failed to get IoStoreOnDemand module!");
+			UE::GameFeatures::FResult ErrorResult = MakeError(FString::Printf(TEXT("%.*s%s"),
+				PredownloadErrorNamespace.Len(), PredownloadErrorNamespace.GetData(),
+				TEXT("IoStoreOnDemand.ModuleNotFound")));
+			SetComplete(MoveTemp(ErrorResult));
+			Cleanup();
+			return;
+		}
+
+		// Create a single list of all IAD assets to install
+		TMap<FName, EStreamingAssetInstallMode> InstallBundlesWithAssetDependencies;
+		for (TPair<FGameFeaturePluginIdentifier, FGFPData>& Pair : GFPs)
+		{
+			FGFPData& GFPData = Pair.Value;
+
+			for (int i = 0; const FName BundleName : GFPData.ProtocolMetadata.InstallBundlesWithAssetDependencies)
+			{
+				const EStreamingAssetInstallMode GfpInstallMode = GFPData.BundleInstallModes[i++];
+			
+				// Merge install modes from all GFPs that references the bundle
+				EStreamingAssetInstallMode& BundleInstallMode = 
+					InstallBundlesWithAssetDependencies.FindOrAdd(BundleName, GfpInstallMode);
+				if (BundleInstallMode == EStreamingAssetInstallMode::GfpRequiredOnly && GfpInstallMode == EStreamingAssetInstallMode::Full)
+				{
+					BundleInstallMode = EStreamingAssetInstallMode::Full;
+				}
+			}
+		}
+
+		IoStorePendingInstalls = InstallBundlesWithAssetDependencies.Num();
+		IoStoreContentHandles.SetNum(IoStorePendingInstalls);
+		IoStoreInstallRequests.SetNum(IoStorePendingInstalls);
+		IoStoreProgress.SetNum(IoStorePendingInstalls);
+		for (int i = 0; const TPair<FName, EStreamingAssetInstallMode>& Pair : InstallBundlesWithAssetDependencies)
+		{
+			const FName InstallBundle = Pair.Key;
+			const EStreamingAssetInstallMode InstallMode = Pair.Value;
+
+			FNameBuilder Debugname(InstallBundle);
+			Debugname.Append(TEXTVIEW("/deps"));
+			IoStoreContentHandles[i] = UE::IoStore::FOnDemandContentHandle::Create(Debugname);
+
+			UE::IoStore::FOnDemandInstallArgs InstallArgs;
+			InstallArgs.MountId = InstallBundle.ToString();
+			if (InstallMode == EStreamingAssetInstallMode::GfpRequiredOnly)
+			{
+				InstallArgs.TagSets.Emplace(TEXTVIEW("required"));
+			}
+			InstallArgs.Options |= UE::IoStore::EOnDemandInstallOptions::InstallSoftReferences;
+
+			InstallArgs.ContentHandle = IoStoreContentHandles[i];
+
+			UE::TUniqueLock Lock(IoStoreMutex);
+
+			IoStoreInstallRequests[i] = IoStore->Install(MoveTemp(InstallArgs),
+				// On Complete
+				[this, Context = SharedThis(this), i](const UE::IoStore::FOnDemandInstallResult& OnDemandInstallResult)
+				{
+					UE::TUniqueLock Lock(IoStoreMutex);
+
+					if (!OnDemandInstallResult.Status.IsOk() && IoStoreStatus.IsOk())
+					{
+						IoStoreStatus = OnDemandInstallResult.Status;
+
+						CancelIoStoreRequests();
+					}
+
+					--IoStorePendingInstalls;
+					if (IoStorePendingInstalls == 0)
+					{
+						ExecuteOnGameThread(UE_SOURCE_LOCATION, [Context]
+						{
+							Context->OnAllIoStoreRequestsCompleted();
+						});
+					}
+				},
+				// On Progress
+				[this, Context = SharedThis(this), i](const UE::IoStore::FOnDemandInstallProgress& InProgress)
+				{
+					UE::TUniqueLock Lock(IoStoreMutex);
+					IoStoreProgress[i] = InProgress;
+				}
+			);
+
+			++i;
+		}
+		
+	}
+
+	void OnAllIoStoreRequestsCompleted()
+	{
+		// Don't need to take the mutex here, all the requests should be done
+
+		if (bCanceled)
+		{
+			SetCompleteCanceled();
+		}
+		else if (!IoStoreStatus.IsOk())
+		{
+			UE_LOGFMT(LogGameFeatures, Error, "GFP Predownload failed to install iostore on demand assets, Error: {Error}", ("Error", IoStoreStatus.ToString()));
+
+			const FString StatusString = IoStoreStatus.ToString();
+			FString ErrorCode = GetIoErrorText(IoStoreStatus.GetErrorCode());
+			ErrorCode.ReplaceCharInline(TEXT(' '), TEXT('_'), ESearchCase::CaseSensitive);
+			
+			UE::GameFeatures::FResult ErrorResult(MakeError(
+				FString::Printf(TEXT("%.*s%s"), PredownloadErrorNamespace.Len(), PredownloadErrorNamespace.GetData(), *ErrorCode)), 
+				FText::AsCultureInvariant(StatusString)); // IOStore errors aren't localized 
+			SetComplete(MoveTemp(ErrorResult));
+		}
+		else
+		{
+			SetComplete();
+		}
+
+		Cleanup();
+	}
+
+	void UpdateProgress()
+	{
+		ProgressTracker.ForceTick();
+		const float InstallBundleProgress = ProgressTracker.GetCurrentCombinedProgress().ProgressPercent;
+
+		if (!bHasAssetDependencies)
+		{
+			Progress = InstallBundleProgress;
+			return;
+		}
+
+		float IoStoreRelativeProgress;
+		{
+			UE::TUniqueLock Lock(IoStoreMutex);
+
+#if 0
+			// This approach doesn't work becuse we don't get initial progress for a request until the request is started
+			// so the total size changes as requests are processed which causes progress backsliding
+			const UE::IoStore::FOnDemandInstallProgress SumIoStoreProgress = Algo::Accumulate(
+				IoStoreProgress,
+				UE::IoStore::FOnDemandInstallProgress(),
+				&UE::IoStore::FOnDemandInstallProgress::Combine);
+
+			IoStoreRelativeProgress = SumIoStoreProgress.GetRelativeProgress();
+#else
+			// TODO: Just use flat weighting for now, but this could be improved by weighting by install size
+			const float SumIoStoreProgress = Algo::TransformAccumulate(
+				IoStoreProgress, 
+				&UE::IoStore::FOnDemandInstallProgress::GetRelativeProgress, 
+				0.0f);
+
+			IoStoreRelativeProgress = IoStoreProgress.IsEmpty() ? 0.0f : (SumIoStoreProgress / IoStoreProgress.Num());
+#endif
+		}
+		
+		Progress = UE::GameFeatures::CombineInstallProgress(InstallBundleProgress, IoStoreRelativeProgress);
+
+		if (OnProgress)
+		{
+			OnProgress(Progress);
+		}
+	}
+
+	void CancelIoStoreRequests()
+	{
+		for (UE::IoStore::FOnDemandInstallRequest& InstallRequest : IoStoreInstallRequests)
+		{
+			InstallRequest.Cancel();
+		}
+	}
+
 	// Predownload shouldn't pin any cached bundles so release them now
 	void ReleaseBundlesIfPossible()
 	{
+		// Don't need to do anything special with content handles. GFPs will keep their own handles internally.
+		IoStoreContentHandles.Empty();
+
 		UGameFeaturesSubsystem& GFPSubSys = UGameFeaturesSubsystem::Get();
 		TSharedPtr<IInstallBundleManager> BundleManager = IInstallBundleManager::GetPlatformInstallBundleManager();
 
@@ -3077,22 +3366,26 @@ struct FGameFeaturePluginPredownloadContext : public FGameFeaturePluginPredownlo
 				}
 			}
 		}
-		for (const TPair<FGameFeaturePluginIdentifier, FInstallBundlePluginProtocolMetaData>& Pair : GFPs)
+
+		for (const TPair<FGameFeaturePluginIdentifier, FGFPData>& Pair : GFPs)
 		{
-			UGameFeaturePluginStateMachine* Machine = GFPSubSys.FindGameFeaturePluginStateMachine(Pair.Key);
-			if (Machine
-				&& Machine->GetCurrentState() > EGameFeaturePluginState::StatusKnown
-				&& Machine->GetCurrentState() != EGameFeaturePluginState::Releasing)
+			const FGameFeaturePluginIdentifier& GFPIdentifier = Pair.Key;
+			const FGFPData& GFPData = Pair.Value;
+
+			UGameFeaturePluginStateMachine* Machine = GFPSubSys.FindGameFeaturePluginStateMachine(GFPIdentifier);
+			if (Machine &&
+				Machine->GetCurrentState() > EGameFeaturePluginState::StatusKnown &&
+				Machine->GetCurrentState() != EGameFeaturePluginState::Releasing)
 			{
 				// A machine is using the bundles, don't release
-				KeepList.Append(Pair.Value.InstallBundles);
+				KeepList.Append(GFPData.ProtocolMetadata.InstallBundles);
 				continue;
 			}
 
 			// Only send this if a GFP is not using the bundle
-			GFPSubSys.OnGameFeaturePostPredownloading(FString(Pair.Key.GetPluginName()), Pair.Key);
+			GFPSubSys.OnGameFeaturePostPredownloading(FString(GFPIdentifier.GetPluginName()), GFPIdentifier);
 
-			ReleaseList.Append(Pair.Value.InstallBundles);
+			ReleaseList.Append(GFPData.ProtocolMetadata.InstallBundles);
 		}
 
 		BundleManager->RequestReleaseContent(ReleaseList, EInstallBundleReleaseRequestFlags::None, KeepList.Array());
diff --git a/Engine/Plugins/Runtime/GameFeatures/Source/GameFeatures/Public/GameFeaturesSubsystem.h b/Engine/Plugins/Runtime/GameFeatures/Source/GameFeatures/Public/GameFeaturesSubsystem.h
--- a/Engine/Plugins/Runtime/GameFeatures/Source/GameFeatures/Public/GameFeaturesSubsystem.h
+++ b/Engine/Plugins/Runtime/GameFeatures/Source/GameFeatures/Public/GameFeaturesSubsystem.h
@@ -5,7 +5,6 @@
 #include "Containers/Union.h"
 #include "Engine/Engine.h"
 #include "GameFeatureTypesFwd.h"
-#include "IO/IoStoreOnDemand.h"
 #include "Misc/TransactionallySafeRWLock.h"
 
 #include "GameFeaturesSubsystem.generated.h"
@@ -693,7 +692,8 @@ public:
 	/** 
 	 * Pre-install any required game feature data, which can be useful for larger payloads. 
 	 * This does not instantiate any GFP although it is safe to do so before this finishes. 
-	*/
+	 * This doesn't not resolve any dependencies, PluginURLs should contain all dependencies.
+	 */
 	UE_API TSharedRef<FGameFeaturePluginPredownloadHandle> PredownloadGameFeaturePlugins(TConstArrayView<FString> PluginURLs, 
 		TUniqueFunction<void(const UE::GameFeatures::FResult&)> OnComplete = nullptr, TUniqueFunction<void(float)> OnProgress = nullptr);
 	friend struct FGameFeaturePluginPredownloadContext;
