Guarantee that OnDeviceDestroyed provides a valid handle to RemoveFromDevice in UGameFeatureAction_AudioActionBase

#rnx

diff --git a/Engine/Plugins/Runtime/GameFeatures/Source/GameFeatures/Private/GameFeatureAction_AudioActionBase.cpp b/Engine/Plugins/Runtime/GameFeatures/Source/GameFeatures/Private/GameFeatureAction_AudioActionBase.cpp
--- a/Engine/Plugins/Runtime/GameFeatures/Source/GameFeatures/Private/GameFeatureAction_AudioActionBase.cpp
+++ b/Engine/Plugins/Runtime/GameFeatures/Source/GameFeatures/Private/GameFeatureAction_AudioActionBase.cpp
@@ -49,7 +49,11 @@ void UGameFeatureAction_AudioActionBase::OnDeviceDestroyed(Audio::FDeviceId InDe
 {
 	if (FAudioDeviceManager* AudioDeviceManager = FAudioDeviceManager::Get())
 	{
-		RemoveFromDevice(AudioDeviceManager->GetAudioDevice(InDeviceId));
+		FAudioDeviceHandle AudioDeviceHandle = AudioDeviceManager->GetAudioDevice(InDeviceId);
+		if (AudioDeviceHandle.IsValid())
+		{
+			RemoveFromDevice(AudioDeviceHandle);
+		}
 	}
 }
 
