GFP error logging improvements
- Use Shortened names to make log more readable
- Logs mentioning dependecies all now include the both the referencing plugin and the dependency
- Add more explicit logging for failing to load GFD for a plugin (add logs for missing package or print the error on the streamble handle if present)

diff --git a/Engine/Plugins/Runtime/GameFeatures/Source/GameFeatures/Private/GameFeaturePluginStateMachine.cpp b/Engine/Plugins/Runtime/GameFeatures/Source/GameFeatures/Private/GameFeaturePluginStateMachine.cpp
--- a/Engine/Plugins/Runtime/GameFeatures/Source/GameFeatures/Private/GameFeaturePluginStateMachine.cpp
+++ b/Engine/Plugins/Runtime/GameFeatures/Source/GameFeatures/Private/GameFeaturePluginStateMachine.cpp
@@ -784,15 +784,23 @@ struct FTransitionDependenciesGameFeaturePluginState : public FGameFeaturePlugin
 					TEXT("FGameFeaturePluginState::AllowAsyncLoading is false for %s but true for dependency being waited on %s"), 
 					*StateProperties.PluginName, *Dependency->GetPluginURL());
 
-				const EGameFeaturePluginState DependencyState = Dependency->GetCurrentState();
-				const bool bIsInErrorState = DependencyState == EGameFeaturePluginState::ErrorRegistering || DependencyState == EGameFeaturePluginState::ErrorWaitingForDependencies;
-				if (TransitionPolicy::ShouldWaitForDependencies() && bIsInErrorState && Dependency->IsErrorStateUnrecoverable())
+				if (TransitionPolicy::ShouldWaitForDependencies() && 
+					Dependency->IsInErrorState() &&
+					Dependency->IsErrorStateUnrecoverable())
 				{
-					UE_LOG(LogGameFeatures, Error, TEXT("Dependency %s failed and is in %s"), *Dependency->GetPluginURL(), *UE::GameFeatures::ToString(DependencyState));
+					const EGameFeaturePluginState DependencyState = Dependency->GetCurrentState();
+					const FStringView ShortUrl = StateProperties.PluginIdentifier.GetIdentifyingString();
+					const FStringView DepShortUrl = Dependency->GetPluginIdentifier().GetIdentifyingString();
+					UE_LOG(LogGameFeatures, Error, TEXT("Plugin (%.*s) dependency (%.*s) failed and is in state %s"),
+						ShortUrl.Len(), ShortUrl.GetData(),
+						DepShortUrl.Len(), DepShortUrl.GetData(),
+						*UE::GameFeatures::ToString(DependencyState));
 					StateStatus.SetTransitionError(TransitionPolicy::GetErrorState(), GetErrorResult(TEXT("Failed_Dependency_Transition")));
 					if (UGameFeaturePluginStateMachine* CurrentMachine = GameFeaturesSubsystem.FindGameFeaturePluginStateMachine(StateProperties.PluginIdentifier))
 					{
-						UE_LOG(LogGameFeatures, Error, TEXT("Setting %s to be in unrecoverable error as one of their dependencies is in unrecoverable error"), *StateProperties.PluginIdentifier.GetFullPluginURL());
+						UE_LOG(LogGameFeatures, Error, TEXT("Setting plugin (%.*s) to be in unrecoverable error because dependency (%.*s) is in unrecoverable error"), 
+							ShortUrl.Len(), ShortUrl.GetData(),
+							DepShortUrl.Len(), DepShortUrl.GetData());
 						CurrentMachine->SetUnrecoverableError();
 					}
 					return;
@@ -817,13 +825,20 @@ struct FTransitionDependenciesGameFeaturePluginState : public FGameFeaturePlugin
 
 			if (Pair.Value.HasError())
 			{
-				UE_LOG(LogGameFeatures, Error, TEXT("Dependency %s failed to transition with error %s"), *RemainingDependency->GetPluginURL(), *Pair.Value.GetError());
+				const FStringView ShortUrl = StateProperties.PluginIdentifier.GetIdentifyingString();
+				const FStringView DepShortUrl = RemainingDependency->GetPluginIdentifier().GetIdentifyingString();
+				UE_LOG(LogGameFeatures, Error, TEXT("Plugin (%.*s) dependency (%.*s) failed to transition with error %s"),
+					ShortUrl.Len(), ShortUrl.GetData(),
+					DepShortUrl.Len(), DepShortUrl.GetData(),
+					*Pair.Value.GetError());
 				StateStatus.SetTransitionError(TransitionPolicy::GetErrorState(), GetErrorResult(TEXT("Failed_Dependency_Transition")));
 				if (RemainingDependency->IsErrorStateUnrecoverable())
 				{
 					if (UGameFeaturePluginStateMachine* CurrentMachine = GameFeaturesSubsystem.FindGameFeaturePluginStateMachine(StateProperties.PluginIdentifier))
 					{
-						UE_LOG(LogGameFeatures, Error, TEXT("Setting %s to be in unrecoverable error as one of their dependencies is in unrecoverable error"), *StateProperties.PluginIdentifier.GetFullPluginURL());
+						UE_LOG(LogGameFeatures, Error, TEXT("Setting plugin (%.*s) to be in unrecoverable error because dependency (%.*s) is in unrecoverable error"),
+							ShortUrl.Len(), ShortUrl.GetData(),
+							DepShortUrl.Len(), DepShortUrl.GetData());
 						CurrentMachine->SetUnrecoverableError();
 					}
 				}
@@ -3115,76 +3130,77 @@ struct FGameFeaturePluginState_Registering : public FGameFeaturePluginState
 	};
 
 	TSharedPtr<FStreamableHandle> GameFeatureDataHandle;
-	TArray<FString, TInlineAllocator<2>> GameFeatureDataSearchPaths;
 	ELoadGFDState LoadGFDState = ELoadGFDState::Pending;
 	bool bCheckedRealtimeMode = false;
 
 	FGameFeaturePluginState_Registering(FGameFeaturePluginStateMachineProperties& InStateProperties) : FGameFeaturePluginState(InStateProperties) {}
 	
-	void TryAsyncLoadGameFeatureData(int32 Attempt = 0)
+	void TryAsyncLoadGameFeatureData(const FString& Path)
 	{
-		if (!GameFeatureDataSearchPaths.IsValidIndex(Attempt))
-		{
-			LoadGFDState = ELoadGFDState::Failed;
-			UpdateStateMachineDeferred();
-			return;
-		}
-
 		check(LoadGFDState == ELoadGFDState::Pending);
 
 		bool bIsLoading = false;
 
-		const FString& Path = GameFeatureDataSearchPaths[Attempt];
-		if (FPackageName::DoesPackageExist(Path))
+		GameFeatureDataHandle = UGameFeaturesSubsystem::LoadGameFeatureData(Path, true /*bStartStalled*/);
+		if (GameFeatureDataHandle && GameFeatureDataHandle->IsLoadingInProgress())
 		{
-			GameFeatureDataHandle = UGameFeaturesSubsystem::LoadGameFeatureData(Path, true /*bStartStalled*/);
-			if (GameFeatureDataHandle && GameFeatureDataHandle->IsLoadingInProgress())
-			{
-				// CurrentMachine owns `this`, so use it as a lifetime check for the `this` captured in the lambdas below, as they may be called after `this` is destroyed
-				UGameFeaturePluginStateMachine* CurrentMachine = UGameFeaturesSubsystem::Get().FindGameFeaturePluginStateMachine(StateProperties.PluginIdentifier);
+			// CurrentMachine owns `this`, so use it as a lifetime check for the `this` captured in the lambdas below, as they may be called after `this` is destroyed
+			UGameFeaturePluginStateMachine* CurrentMachine = UGameFeaturesSubsystem::Get().FindGameFeaturePluginStateMachine(StateProperties.PluginIdentifier);
 
-				GameFeatureDataHandle->BindCancelDelegate(FStreamableDelegateWithHandle::CreateWeakLambda(CurrentMachine, [this](const TSharedPtr<FStreamableHandle>& Handle)
+			GameFeatureDataHandle->BindCancelDelegate(FStreamableDelegateWithHandle::CreateWeakLambda(CurrentMachine, [this](const TSharedPtr<FStreamableHandle>& Handle)
+			{
+				if (GameFeatureDataHandle != Handle)
 				{
-					if (GameFeatureDataHandle != Handle)
-					{
-						// We're no longer in a state where we want to process this callback (maybe we already went through EndState)
-						return;
-					}
+					// We're no longer in a state where we want to process this callback (maybe we already went through EndState)
+					return;
+				}
 
-					const FStringView ShortUrl = StateProperties.PluginIdentifier.GetIdentifyingString();
-					UE_LOG(LogGameFeatures, Error, TEXT("Game Feature Data loading was cancelled for URL %.*s"), ShortUrl.Len(), ShortUrl.GetData());
+				const FStringView ShortUrl = StateProperties.PluginIdentifier.GetIdentifyingString();
+				UE_LOG(LogGameFeatures, Error, TEXT("Game Feature Data loading was cancelled for URL %.*s"), ShortUrl.Len(), ShortUrl.GetData());
 
-					LoadGFDState = ELoadGFDState::Cancelled;
-					UpdateStateMachineDeferred();
-				}));
+				LoadGFDState = ELoadGFDState::Cancelled;
+				UpdateStateMachineDeferred();
+			}));
+
+			GameFeatureDataHandle->BindCompleteDelegate(FStreamableDelegateWithHandle::CreateWeakLambda(CurrentMachine, [this](const TSharedPtr<FStreamableHandle>& Handle)
+			{
+				if (GameFeatureDataHandle != Handle)
+				{
+					// We're no longer in a state where we want to process this callback (maybe we already went through EndState)
+					return;
+				}
+
+				StateProperties.GameFeatureData = Cast<UGameFeatureData>(GameFeatureDataHandle->GetLoadedAsset());
 
-				GameFeatureDataHandle->BindCompleteDelegate(FStreamableDelegateWithHandle::CreateWeakLambda(CurrentMachine, [this, Attempt](const TSharedPtr<FStreamableHandle>& Handle)
+				if (StateProperties.GameFeatureData)
+				{
+					LoadGFDState = ELoadGFDState::Success;
+				}
+				else
 				{
-					if (GameFeatureDataHandle != Handle)
+					LoadGFDState = ELoadGFDState::Failed;
+
+					const FStringView ShortUrl = StateProperties.PluginIdentifier.GetIdentifyingString();
+					if (const TOptional<UE::UnifiedError::FError>& Error = GameFeatureDataHandle->GetError())
 					{
-						// We're no longer in a state where we want to process this callback (maybe we already went through EndState)
-						return;
+						UE_LOG(LogGameFeatures, Error, TEXT("Game Feature Data loading failed with error [%s] for URL %.*s"), *LexToString(*Error), ShortUrl.Len(), ShortUrl.GetData());
 					}
-
-					StateProperties.GameFeatureData = Cast<UGameFeatureData>(GameFeatureDataHandle->GetLoadedAsset());
-					if (!StateProperties.GameFeatureData)
+					else
 					{
-						TryAsyncLoadGameFeatureData(Attempt + 1);
-						return;
+						UE_LOG(LogGameFeatures, Error, TEXT("Game Feature Data loading failed without an error for URL %.*s"), ShortUrl.Len(), ShortUrl.GetData());
 					}
+				}
+				
+				UpdateStateMachineDeferred();
+			}));
 
-					LoadGFDState = ELoadGFDState::Success;
-					UpdateStateMachineDeferred();
-				}));
-
-				bIsLoading = true;
-				GameFeatureDataHandle->StartStalledHandle();
-			}
+			bIsLoading = true;
+			GameFeatureDataHandle->StartStalledHandle();
 		}
 
 		if (!bIsLoading)
 		{
-			TryAsyncLoadGameFeatureData(Attempt + 1);
+			LoadGFDState = ELoadGFDState::Failed;
 		}
 	}
 
@@ -3251,13 +3267,27 @@ struct FGameFeaturePluginState_Registering : public FGameFeaturePluginState
 			}
 		}
 
-		GameFeatureDataSearchPaths.Empty();
-		GameFeatureDataSearchPaths.Emplace(MoveTemp(PreferredGameFeatureDataPath));
-		GameFeatureDataSearchPaths.Emplace(MoveTemp(BackupGameFeatureDataPath));
+		FString GameFeatureDataPath;
+		if (FPackageName::DoesPackageExist(PreferredGameFeatureDataPath))
+		{
+			GameFeatureDataPath = MoveTemp(PreferredGameFeatureDataPath);
+		}
+		else if (FPackageName::DoesPackageExist(BackupGameFeatureDataPath))
+		{
+			GameFeatureDataPath = MoveTemp(BackupGameFeatureDataPath);
+		}
+		else
+		{
+			const FStringView ShortUrl = StateProperties.PluginIdentifier.GetIdentifyingString();
+			UE_LOG(LogGameFeatures, Error, TEXT("Game Feature Data package not found for URL %.*s"), ShortUrl.Len(), ShortUrl.GetData());
+
+			LoadGFDState = ELoadGFDState::Failed;
+			return;
+		}
 
 		if (UseAsyncLoading())
 		{
-			TryAsyncLoadGameFeatureData();
+			TryAsyncLoadGameFeatureData(GameFeatureDataPath);
 		}
 		else
 		{
@@ -3266,25 +3296,32 @@ struct FGameFeaturePluginState_Registering : public FGameFeaturePluginState
 					LOCTEXT("LoadingGameFeatureData", "Loading Game Feature Data for Plugin: {0}"), 
 					FText::FromString(StateProperties.PluginName)));
 			LoadingGameFeatureData.Visibility = ESlowTaskVisibility::Important;
-			for (const FString& Path : GameFeatureDataSearchPaths)
+			
+			GameFeatureDataHandle = UGameFeaturesSubsystem::LoadGameFeatureData(GameFeatureDataPath);
+			if (GameFeatureDataHandle)
 			{
-				if (FPackageName::DoesPackageExist(Path))
+				GameFeatureDataHandle->WaitUntilComplete(0.0f, false);
+				StateProperties.GameFeatureData = Cast<UGameFeatureData>(GameFeatureDataHandle->GetLoadedAsset());
+			}
+
+			if (StateProperties.GameFeatureData)
+			{
+				LoadGFDState = ELoadGFDState::Success;
+			}
+			else
+			{
+				LoadGFDState = ELoadGFDState::Failed;
+
+				const FStringView ShortUrl = StateProperties.PluginIdentifier.GetIdentifyingString();
+				if (const TOptional<UE::UnifiedError::FError>& Error = GameFeatureDataHandle->GetError())
 				{
-					GameFeatureDataHandle = UGameFeaturesSubsystem::LoadGameFeatureData(Path);
-					if (GameFeatureDataHandle)
-					{
-						GameFeatureDataHandle->WaitUntilComplete(0.0f, false);
-						StateProperties.GameFeatureData = Cast<UGameFeatureData>(GameFeatureDataHandle->GetLoadedAsset());
-					}
+					UE_LOG(LogGameFeatures, Error, TEXT("Game Feature Data loading failed with error [%s] for URL %.*s"), *LexToString(*Error), ShortUrl.Len(), ShortUrl.GetData());
 				}
-
-				if (StateProperties.GameFeatureData)
+				else
 				{
-					break;
+					UE_LOG(LogGameFeatures, Error, TEXT("Game Feature Data loading failed without an error for URL %.*s"), ShortUrl.Len(), ShortUrl.GetData());
 				}
 			}
-
-			LoadGFDState = StateProperties.GameFeatureData ? ELoadGFDState::Success : ELoadGFDState::Failed;
 		}
 	}
 
@@ -3342,7 +3379,8 @@ struct FGameFeaturePluginState_Registering : public FGameFeaturePluginState
 			StateStatus.SetTransitionError(EGameFeaturePluginState::ErrorRegistering, GetErrorResult(TEXT("Plugin_Missing_GameFeatureData")));
 			if (UGameFeaturePluginStateMachine* CurrentMachine = UGameFeaturesSubsystem::Get().FindGameFeaturePluginStateMachine(StateProperties.PluginIdentifier))
 			{
-				UE_LOG(LogGameFeatures, Error, TEXT("Setting %s to be in unrecoverable error as GameFeatureData is missing"), *StateProperties.PluginIdentifier.GetFullPluginURL());
+				const FStringView ShortUrl = StateProperties.PluginIdentifier.GetIdentifyingString();
+				UE_LOG(LogGameFeatures, Error, TEXT("Setting %.*s to be in unrecoverable error as GameFeatureData is missing"), ShortUrl.Len(), ShortUrl.GetData());
 				CurrentMachine->SetUnrecoverableError();
 			}
 		}
@@ -4252,6 +4290,11 @@ bool UGameFeaturePluginStateMachine::IsAvailable() const
 	return GetCurrentState() >= EGameFeaturePluginState::StatusKnown;
 }
 
+bool UGameFeaturePluginStateMachine::IsInErrorState() const
+{
+	return IsValidErrorState(GetCurrentState());
+}
+
 bool UGameFeaturePluginStateMachine::AllowAsyncLoading() const
 {
 	return StateProperties.AllowAsyncLoading();
diff --git a/Engine/Plugins/Runtime/GameFeatures/Source/GameFeatures/Private/GameFeaturePluginStateMachine.h b/Engine/Plugins/Runtime/GameFeatures/Source/GameFeatures/Private/GameFeaturePluginStateMachine.h
--- a/Engine/Plugins/Runtime/GameFeatures/Source/GameFeatures/Private/GameFeaturePluginStateMachine.h
+++ b/Engine/Plugins/Runtime/GameFeatures/Source/GameFeatures/Private/GameFeaturePluginStateMachine.h
@@ -556,6 +556,9 @@ public:
 	/** Returns true if the plugin is available to download/load. Only call if IsStatusKnown is true */
 	bool IsAvailable() const;
 
+	/** Returns true if the plugin is in an error state */
+	bool IsInErrorState() const;
+
 	/** Whether this machine is allowed to be asynchronous */
 	bool AllowAsyncLoading() const;
 
