Gather stats for package loading resulting from GFPs
- Use existing package access tracking facility to record all package loads that happen as a result of GFP state transitions.  Tracking scopes are pushed onto the stack from within the game feature state machine. Recorded stats are the tree of package loads and the operations that caused the loads. Works for sync, async, direct, and indirect loads (unless they are deferred and loaded dynamically from code).
- Packages that are directly loaded have an operation name `GFP_<state>` (e.g. `GFP_Registering`). Indirectly loaded packages (e.g. imports/dependencies) use the existing operation names, most typically `Load`.
- Disabled by default, enabled with `GameFeaturePlugin.TrackLoadStats`, and printed with `GameFeaturePlugin.PrintLoadStats`.
- Removed some unused includes from `PackageAccessTracking.h`.
- Potential future improvements
- Improve the tracking system to track load requests instead of package creation. This avoids the need to throw away the context stack for async loads like the system currently does. It also means we can detect when a package is loaded by multiple GFPs and track it for each.
- Add tracking scopes for untracked loads. This maybe benefit cook tracking as well since it uses the same system. Most of these are plugins and built-in assets that have hard coded loads in the engine. I didn't see any packages that appeared to be related to GFPs.
- Flush loading during GFP state changes and capture all package creation. This would catch "untracked" loads, at the expense perf and added complexity.
- Improve the tracking system to allow opt-in for specific packages without needing to store load stats in the calling system, which will result in duplicate data when multiple parties are interested in stats.

diff --git a/Engine/Plugins/Runtime/GameFeatures/Source/GameFeatures/Private/GameFeaturePluginStateMachine.cpp b/Engine/Plugins/Runtime/GameFeatures/Source/GameFeatures/Private/GameFeaturePluginStateMachine.cpp
--- a/Engine/Plugins/Runtime/GameFeatures/Source/GameFeatures/Private/GameFeaturePluginStateMachine.cpp
+++ b/Engine/Plugins/Runtime/GameFeatures/Source/GameFeatures/Private/GameFeaturePluginStateMachine.cpp
@@ -526,6 +526,25 @@ namespace UE::GameFeatures
 
 	using FBroadcastingOnDestinationStateReached = TBroadcastingCallback<FDestinationGameFeaturePluginState::FOnDestinationStateReached>;
 	using FBroadcastingOnTransitionCanceled = TBroadcastingCallback<FGameFeaturePluginStateMachineProperties::FOnTransitionCanceled>;
+
+	static FName GetStateName(EGameFeaturePluginState State)
+	{
+		switch (State)
+		{
+			#define X(inEnum, inText) \
+				case EGameFeaturePluginState::inEnum: \
+				{ \
+					static const FName Name_##inEnum = FName(TEXT("GFP_") #inEnum); \
+					return Name_##inEnum; \
+				}
+
+			GAME_FEATURE_PLUGIN_STATE_LIST(X)
+
+			#undef X
+		}
+
+		return FName(TEXT("Unknown"));
+	};
 }
 
 void FGameFeaturePluginStateStatus::SetTransition(EGameFeaturePluginState InTransitionToState)
@@ -693,7 +712,7 @@ bool FGameFeaturePluginState::AllowIniLoading() const
 	{
 	case (EGameFeaturePluginProtocol::InstallBundle):
 	{
-		//InstallBundleProtocol's have a MetaData that controlls if INI loading is allowed
+		//InstallBundleProtocol's have a MetaData that controls if INI loading is allowed
 		//The protocol default is not to allow INI loading since the source is likely untrusted
 		return StateProperties.ProtocolOptions.GetSubtype<FInstallBundlePluginProtocolOptions>().bAllowIniLoading;
 	}
@@ -4432,6 +4451,10 @@ void UGameFeaturePluginStateMachine::UpdateStateMachine()
 
 		FGameFeaturePluginStateStatus StateStatus;
 		{
+			const FGameFeaturePluginStateMachineProperties& Props = AllStates[CurrentState]->StateProperties;
+			FName LoaderName = Props.GameFeatureData ? Props.GameFeatureData->GetPackage()->GetFName() : FName(Props.PluginName);
+			UE_TRACK_REFERENCING_PACKAGE_SCOPED(LoaderName, UE::GameFeatures::GetStateName(CurrentState));
+
 			TRACE_CPUPROFILER_EVENT_SCOPE(GFP_UpdateStateMachine_UpdateState);
 			AllStates[CurrentState]->UpdateState(StateStatus);
 		}
@@ -4452,6 +4475,11 @@ void UGameFeaturePluginStateMachine::UpdateStateMachine()
 			CurrentStateInfo = FGameFeaturePluginStateInfo(StateStatus.TransitionToState);
 			CurrentState = StateStatus.TransitionToState;
 			check(CurrentState != EGameFeaturePluginState::MAX);
+
+			const FGameFeaturePluginStateMachineProperties& Props = AllStates[CurrentState]->StateProperties;
+			FName LoaderName = Props.GameFeatureData ? Props.GameFeatureData->GetPackage()->GetFName() : FName(Props.PluginName);
+			UE_TRACK_REFERENCING_PACKAGE_SCOPED(LoaderName, UE::GameFeatures::GetStateName(CurrentState));
+
 			{
 				TRACE_CPUPROFILER_EVENT_SCOPE(GFP_UpdateStateMachine_BeginState);
 				CheckAddBatchingRequestForCurrentState();
diff --git a/Engine/Plugins/Runtime/GameFeatures/Source/GameFeatures/Private/GameFeaturesSubsystem.cpp b/Engine/Plugins/Runtime/GameFeatures/Source/GameFeatures/Private/GameFeaturesSubsystem.cpp
--- a/Engine/Plugins/Runtime/GameFeatures/Source/GameFeatures/Private/GameFeaturesSubsystem.cpp
+++ b/Engine/Plugins/Runtime/GameFeatures/Source/GameFeatures/Private/GameFeaturesSubsystem.cpp
@@ -111,6 +111,13 @@ namespace UE::GameFeatures
 		TEXT("When call LoadBuiltInGameFeaturePlugins should only the plugins that were built by this target exe be considered built in."),
 		ECVF_Default);
 
+	static FAutoConsoleVariable CVarTrackLoadStats(
+		TEXT("GameFeaturePlugin.TrackLoadStats"),
+		false,
+		TEXT("Track all packages loaded by each GFP. If a package is loaded by multiple GFPs, it will only be ")
+		TEXT("attributed to the first one that actually causes the load. Stats can be output with ")
+		TEXT("GameFeaturePlugin.PrintLoadStats."));
+
 	TOptional<FString> GetPluginUrlForConsoleCommand(const TArray<FString>& Args, FOutputDevice& Ar)
 	{
 		TOptional<FString> PluginURL;
@@ -142,6 +149,128 @@ namespace UE::GameFeatures
 
 		return PluginURL;
 	}
+
+#if UE_WITH_PACKAGE_ACCESS_TRACKING
+	class FPackageLoadTracker final : public FUObjectArray::FUObjectCreateListener
+	{
+		struct FStat
+		{
+			FName Package;
+			FName Op;
+		};
+
+		/**
+		 * Store which packages caused other packages to be loaded. Only tracks loads that occur due to a GFP (direct
+		 * and indirect, sync and async). Keys are a package name and values are a list of package names directly loaded
+		 * by said package. Indirect loads are gathered by walking the list of loaded packages and querying the map
+		 * recursively.
+		 */
+		TArray<FName> Roots;
+		TMap<FName, TArray<FStat>> Stats;
+		TArray<FName> UntrackedLoads;
+
+	public:
+
+		FPackageLoadTracker()
+		{
+			GUObjectArray.AddUObjectCreateListener(this);
+		}
+
+		~FPackageLoadTracker()
+		{
+			GUObjectArray.RemoveUObjectCreateListener(this);
+		}
+
+		virtual void OnUObjectArrayShutdown() override
+		{
+			GUObjectArray.RemoveUObjectCreateListener(this);
+		}
+
+		virtual void NotifyUObjectCreated(const UObjectBase* Object, int32 Index) override
+		{
+			if (Object->GetClass() == UPackage::StaticClass())
+			{
+				using namespace PackageAccessTracking_Private;
+
+				const UPackage* Package = Cast<UPackage>(Object);
+				FString Path = Package->GetPathName();
+
+				// Async loading will cache tracking data when a load is requested and restore it when processing the
+				// load later. This means we have tracking for async loads, but it only caches the PackageName and
+				// OpName from FTrackedData so we can only rely on those two fields.
+				const FTrackedData* Data = FPackageAccessRefScope::GetCurrentThreadAccumulatedData();
+				if (Data && !Data->PackageName.IsNone())
+				{
+					FName LoaderName = Data->PackageName;
+					if (TArray<FStat>* LoadedPackages = Stats.Find(LoaderName))
+					{
+						FName LoadeeName = Object->GetFName();
+						LoadedPackages->Add({ LoadeeName, Data->OpName });
+						Stats.Add(LoadeeName);
+					}
+				}
+				else
+				{
+					FName LoadeeName = Object->GetFName();
+					UntrackedLoads.Add(LoadeeName);
+				}
+			}
+		}
+
+		void AddRoot(FName PackageName)
+		{
+			Roots.Add(PackageName);
+			Stats.Add(PackageName);
+		}
+
+		void PrintLoadStats(bool bShowUntracked) const
+		{
+			struct Local
+			{
+				static void VisitPackage(FStringBuilderBase& Buf, const TMap<FName, TArray<FStat>>& Stats, FName Package)
+				{
+					if (const TArray<FStat>* LoadedList = Stats.Find(Package))
+					{
+						for (const FStat& Loaded : *LoadedList)
+						{
+							Buf << TEXT("\n\t\t") << Loaded.Op << TEXT(",");
+							int32 Pad = FMath::Max(0, int32(23 - Loaded.Op.GetStringLength()));
+							Buf.Appendf(TEXT("%*s"), Pad, TEXT(""));
+							Buf << Loaded.Package;
+							VisitPackage(Buf, Stats, Loaded.Package);
+						}
+					}
+				}
+			};
+
+			TStringBuilder<1024> Buf;
+			Buf << TEXT("GameFeaturePlugin Package Load Stats");
+			
+			for (FName Root : Roots)
+			{
+				Buf << TEXT("\n\t") << Root << TEXT(": {");
+				Local::VisitPackage(Buf, Stats, Root);
+				Buf << TEXT("\n\t}");
+			}
+
+			if (bShowUntracked)
+			{
+				Buf << TEXT("\n\tUntracked: {");
+				for (FName Untracked : UntrackedLoads)
+				{
+					Buf << TEXT("\n\t\t") << Untracked;
+				}
+				Buf << TEXT("\n\t}");
+			}
+			else
+			{
+				Buf << TEXT("\n\tUntracked: ") << UntrackedLoads.Num();
+			}
+
+			UE_LOG(LogGameFeatures, Display, TEXT("%s"), Buf.ToString());
+		}
+	};
+#endif
 }
 
 const FString LexToString(const EBuiltInAutoState BuiltInAutoState)
@@ -490,7 +619,7 @@ void UGameFeaturesSubsystem::Initialize(FSubsystemCollectionBase& Collection)
 
 	IConsoleManager::Get().RegisterConsoleCommand(
 		TEXT("DisableDebugGameFeatureState"),
-		TEXT("Tuen off triggering a debug breakpoint if when the state of the gameplay feature changes"),
+		TEXT("Turn off triggering a debug breakpoint if when the state of the gameplay feature changes"),
 		FConsoleCommandWithWorldArgsAndOutputDeviceDelegate::CreateLambda([](const TArray<FString>& Args, UWorld*, FOutputDevice& Ar)
 			{
 				if (TOptional<FString> PluginURL = UE::GameFeatures::GetPluginUrlForConsoleCommand(Args, Ar))
@@ -501,6 +630,25 @@ void UGameFeaturesSubsystem::Initialize(FSubsystemCollectionBase& Collection)
 		ECVF_Cheat);
 #endif
 
+#if UE_WITH_PACKAGE_ACCESS_TRACKING
+	if (UE::GameFeatures::CVarTrackLoadStats->GetBool())
+	{
+		PackageLoadTracker = MakeUnique<UE::GameFeatures::FPackageLoadTracker>();
+		IConsoleManager::Get().RegisterConsoleCommand(
+			TEXT("GameFeaturePlugin.PrintLoadStats"),
+			TEXT("Print a list of all packages loaded by each GFP. If a package is loaded by multiple GFPs, it will only be")
+			TEXT("attributed to the first one that actually causes the load")
+			TEXT("\nUsage: GameFeaturePlugin.PrintLoadStats [-ShowUntracked]")
+			TEXT("\n-ShowUntracked - Print a list of packages that were loaded, but don't have tracking to know why"),
+			FConsoleCommandWithArgsDelegate::CreateLambda([&Tracker = PackageLoadTracker](const TArray<FString>& Args)
+			{
+				bool bShowUntracked = Args.Contains(TEXT("-ShowUntracked"));
+				Tracker->PrintLoadStats(bShowUntracked);
+			})
+		);
+	}
+#endif
+
 	GetExplanationForUnavailablePackageDelegateHandle = IPluginManager::Get().GetExplanationForUnavailablePackageWithPluginInfoDelegate().AddUObject(this, &UGameFeaturesSubsystem::GetExplanationForUnavailablePackage);
 	IPluginManager::Get().OnPluginUnmounted().AddWeakLambda(this, [this] (const IPlugin& Plugin) { PruneCachedGameFeaturePluginDetails(/* doesnt use URL */ FString(""), Plugin.GetDescriptorFileName()); });
 	IPluginManager::Get().OnPluginEdited().AddWeakLambda(this, [this](IPlugin& Plugin) { PruneCachedGameFeaturePluginDetails(/* doesnt use URL */ FString(""), Plugin.GetDescriptorFileName()); });
@@ -2223,7 +2371,7 @@ void UGameFeaturesSubsystem::LoadBuiltInGameFeaturePlugin(const TSharedRef<IPlug
 				BehaviorOptions.AutoStateOverride : PluginDetails.BuiltInAutoState;
 			if (InitialAutoState < EBuiltInAutoState::Registered && IsRunningCookCommandlet())
 			{
-				UE_LOG(LogGameFeatures, Display, TEXT("%s will be set to Registerd for cooking"), *Plugin->GetName());
+				UE_LOG(LogGameFeatures, Display, TEXT("%s will be set to Registered for cooking"), *Plugin->GetName());
 				InitialAutoState = EBuiltInAutoState::Registered;
 			}
 			const EGameFeaturePluginState DestinationState = ConvertInitialFeatureStateToTargetState(InitialAutoState);
@@ -3455,6 +3603,9 @@ UGameFeaturePluginStateMachine* UGameFeaturesSubsystem::FindOrCreateGameFeatureP
 
 	UE_LOG(LogGameFeatures, Verbose, TEXT("Creating GameFeaturePlugin StateMachine using Identifier:%.*s from PluginURL:%s"), PluginIdentifier.GetIdentifyingString().Len(), PluginIdentifier.GetIdentifyingString().GetData(), *PluginURL);
 
+	FName PluginName = FName(PluginIdentifier.GetPluginName());
+	PackageLoadTracker->AddRoot(PluginName);
+
 	UGameFeaturePluginStateMachine* NewStateMachine = NewObject<UGameFeaturePluginStateMachine>(this);
 	GameFeaturePluginStateMachines.Add(FString(PluginIdentifier.GetIdentifyingString()), NewStateMachine);
 	NewStateMachine->InitStateMachine(MoveTemp(PluginIdentifier), ProtocolOptions);
@@ -3495,11 +3646,11 @@ void UGameFeaturesSubsystem::ChangeGameFeatureDestination(UGameFeaturePluginStat
 	}
 	else
 	{
-		FGameFeaturePluginStateRange CurrDesitination = Machine->GetDestination();
+		FGameFeaturePluginStateRange CurrDestination = Machine->GetDestination();
 		UE_LOG(LogGameFeatures, Display, TEXT("ChangeGameFeatureDestination: Attempting to cancel transition for Game Feature %s. Desired [%s, %s]. Current [%s, %s]"), 
 			*Machine->GetGameFeatureName(), 
 			*UE::GameFeatures::ToString(StateRange.MinState), *UE::GameFeatures::ToString(StateRange.MaxState),
-			*UE::GameFeatures::ToString(CurrDesitination.MinState), *UE::GameFeatures::ToString(CurrDesitination.MaxState));
+			*UE::GameFeatures::ToString(CurrDestination.MinState), *UE::GameFeatures::ToString(CurrDestination.MaxState));
 
 		// Try canceling any current transition, then retry
 		auto OnCanceled = [this, InProtocolOptions, StateRange, CompleteDelegate](UGameFeaturePluginStateMachine* Machine) mutable
diff --git a/Engine/Plugins/Runtime/GameFeatures/Source/GameFeatures/Public/GameFeaturesSubsystem.h b/Engine/Plugins/Runtime/GameFeatures/Source/GameFeatures/Public/GameFeaturesSubsystem.h
--- a/Engine/Plugins/Runtime/GameFeatures/Source/GameFeatures/Public/GameFeaturesSubsystem.h
+++ b/Engine/Plugins/Runtime/GameFeatures/Source/GameFeatures/Public/GameFeaturesSubsystem.h
@@ -11,7 +11,11 @@
 
 #define UE_API GAMEFEATURES_API
 
-namespace UE::GameFeatures { struct FResult; }
+namespace UE::GameFeatures
+{
+	class FPackageLoadTracker;
+	struct FResult;
+}
 
 class UGameFeaturePluginStateMachine;
 class IGameFeatureStateChangeObserver;
@@ -912,6 +916,8 @@ private:
 	UPROPERTY(Transient)
 	TObjectPtr<UGameFeaturesProjectPolicies> GameSpecificPolicies;
 
+	TUniquePtr<class UE::GameFeatures::FPackageLoadTracker> PackageLoadTracker;
+
 #if WITH_EDITOR
 	// When we decide not to mount a plugin, we can store an explanation here so that if we later attempt to load an asset from it we can tell the user why it's not available
 	TMap<FString, FString> UnmountedPluginNameToExplanation;
diff --git a/Engine/Source/Runtime/CoreUObject/Private/Misc/PackageAccessTracking.cpp b/Engine/Source/Runtime/CoreUObject/Private/Misc/PackageAccessTracking.cpp
--- a/Engine/Source/Runtime/CoreUObject/Private/Misc/PackageAccessTracking.cpp
+++ b/Engine/Source/Runtime/CoreUObject/Private/Misc/PackageAccessTracking.cpp
@@ -4,6 +4,9 @@
 
 #include "Cooker/CookDependency.h"
 #include "Misc/PackageAccessTrackingOps.h"
+#include "Misc/PackageName.h"
+#include "UObject/Object.h"
+#include "UObject/Package.h"
 
 #if UE_WITH_PACKAGE_ACCESS_TRACKING
 thread_local PackageAccessTracking_Private::FPackageAccessRefScope* PackageAccessTracking_Private::FPackageAccessRefScope::CurrentThreadScope = nullptr;
@@ -41,8 +44,8 @@ FTrackedData::FTrackedData(FTrackedData& DirectData, FTrackedData* OuterAccumula
 	if (OuterAccumulatedData)
 	{
 		// NOTE: The accumulated PackageName is used as the referencer
-		// The referencer is always the PackageName from the latest scope with a PackageName, NOT the PackageName from the latest scope with a BuildOpName. 
-		// Consider a situation where the in the process of loading package A,
+		// The referencer is always the PackageName from the latest scope with a PackageName, NOT the PackageName from the latest scope with a BuildOpName.
+		// Consider a situation where in the process of loading package A,
 		// we call PostLoad on an object on package B which references an object in package C.  The SearchThreadScope with OpName == Load is for package A,
 		// but we want to record that B depends on C, not that A depends on C.  The InnermostThreadScope will have a package name of B and an OpName of PostLoad
 		// and so we want the package name from the InnermostThreadScope while having searched upwards for a scope with an OpName == Load.
diff --git a/Engine/Source/Runtime/CoreUObject/Public/Misc/PackageAccessTracking.h b/Engine/Source/Runtime/CoreUObject/Public/Misc/PackageAccessTracking.h
--- a/Engine/Source/Runtime/CoreUObject/Public/Misc/PackageAccessTracking.h
+++ b/Engine/Source/Runtime/CoreUObject/Public/Misc/PackageAccessTracking.h
@@ -4,13 +4,10 @@
 
 #include "CoreTypes.h"
 #include "Misc/Optional.h"
-#include "Misc/PackageName.h"
 #include "UObject/NameTypes.h"
-#include "UObject/ObjectHandle.h"
-#include "UObject/Package.h"
+#include "UObject/ObjectHandleDefines.h"
 
 class ITargetPlatform;
-class UPackage;
 
 #define UE_WITH_PACKAGE_ACCESS_TRACKING UE_WITH_OBJECT_HANDLE_TRACKING
 
@@ -39,8 +36,8 @@ class UPackage;
 /**
  * Specify that within the current scope, dereferenced TObjectPtrs cause a dependency on the specified
  * ResultProjection of the dereferenced object and package, instead of the outer scope's projection or the default
- * projection of UE::CooK::ResultProjection::All. ProjectionName must be either one of the generic terms provided
- * in UE::Cook::ResultProjection or it must be a name declared as a resultprojection by one of the objects in
+ * projection of UE::Cook::ResultProjection::All. ProjectionName must be either one of the generic terms provided
+ * in UE::Cook::ResultProjection or it must be a name declared as a result projection by one of the objects in
  * the target package.
  */
 #define UE_COOK_RESULTPROJECTION_SCOPED(ProjectionName) \
@@ -60,6 +57,7 @@ namespace PackageAccessTracking_Private
 		/** Standard constructor; sets variables from DirectData passed to a FPackageAccessRefScope. */
 		COREUOBJECT_API FTrackedData(FName PackageNameIn, FName InOpName, FName InCookResultProjection,
 			const ITargetPlatform* InTargetPlatform, const UObject* InObject);
+
 		/** Accumulating constructor; sets new AccumulatedData by combining DirectData with Outer's AccumulatedData. */
 		COREUOBJECT_API FTrackedData(FTrackedData& DirectData, FTrackedData* OuterAccumulatedData);
 
@@ -106,3 +104,9 @@ namespace PackageAccessTracking_Private
 	};
 };
 #endif // UE_WITH_PACKAGE_ACCESS_TRACKING
+
+#if UE_ENABLE_INCLUDE_ORDER_DEPRECATED_IN_5_7
+#include "Misc/PackageName.h"
+#include "UObject/ObjectHandle.h"
+#include "UObject/Package.h"
+#endif
diff --git a/Engine/Source/Runtime/CoreUObject/Public/Misc/PackageAccessTrackingOps.h b/Engine/Source/Runtime/CoreUObject/Public/Misc/PackageAccessTrackingOps.h
--- a/Engine/Source/Runtime/CoreUObject/Public/Misc/PackageAccessTrackingOps.h
+++ b/Engine/Source/Runtime/CoreUObject/Public/Misc/PackageAccessTrackingOps.h
@@ -15,11 +15,13 @@ namespace PackageAccessTrackingOps
 	extern COREUOBJECT_API FName NAME_Save;
 	extern COREUOBJECT_API FName NAME_CreateDefaultObject;
 	extern COREUOBJECT_API FName NAME_CookerBuildObject;
+
 	/**
-	 * A global-scope operation suck as AssetRegistry tick is occurring on the current thread inside of another scope,
-	 * and the references made during this operation should not be associated with the previous scope.
+	 * A global-scope operation started as AssetRegistry tick is occurring on the current thread inside of another
+	 * scope, and the references made during this operation should not be associated with the previous scope.
 	 */
 	extern COREUOBJECT_API FName NAME_ResetContext;
+
 	/**
 	 * A debug-only scope that is used to find global-scope operations. Any accesses made when a NoAccessExpected
 	 * scope is on top of the context stack will log a warning. To resolve the warning either the NoAccessExpected
