Fix compilation issues with 44292583

diff --git a/Engine/Plugins/Runtime/GameFeatures/Source/GameFeatures/Private/GameFeaturesSubsystem.cpp b/Engine/Plugins/Runtime/GameFeatures/Source/GameFeatures/Private/GameFeaturesSubsystem.cpp
--- a/Engine/Plugins/Runtime/GameFeatures/Source/GameFeatures/Private/GameFeaturesSubsystem.cpp
+++ b/Engine/Plugins/Runtime/GameFeatures/Source/GameFeatures/Private/GameFeaturesSubsystem.cpp
@@ -270,6 +270,8 @@ namespace UE::GameFeatures
 			UE_LOG(LogGameFeatures, Display, TEXT("%s"), Buf.ToString());
 		}
 	};
+#else
+	class FPackageLoadTracker {};
 #endif
 }
 
@@ -3603,8 +3605,13 @@ UGameFeaturePluginStateMachine* UGameFeaturesSubsystem::FindOrCreateGameFeatureP
 
 	UE_LOG(LogGameFeatures, Verbose, TEXT("Creating GameFeaturePlugin StateMachine using Identifier:%.*s from PluginURL:%s"), PluginIdentifier.GetIdentifyingString().Len(), PluginIdentifier.GetIdentifyingString().GetData(), *PluginURL);
 
-	FName PluginName = FName(PluginIdentifier.GetPluginName());
-	PackageLoadTracker->AddRoot(PluginName);
+#if UE_WITH_PACKAGE_ACCESS_TRACKING
+	if (PackageLoadTracker)
+	{
+		FName PluginName = FName(PluginIdentifier.GetPluginName());
+		PackageLoadTracker->AddRoot(PluginName);
+	}
+#endif
 
 	UGameFeaturePluginStateMachine* NewStateMachine = NewObject<UGameFeaturePluginStateMachine>(this);
 	GameFeaturePluginStateMachines.Add(FString(PluginIdentifier.GetIdentifyingString()), NewStateMachine);
diff --git a/Engine/Source/Editor/UnrealEd/Private/Cooker/CookSandbox.cpp b/Engine/Source/Editor/UnrealEd/Private/Cooker/CookSandbox.cpp
--- a/Engine/Source/Editor/UnrealEd/Private/Cooker/CookSandbox.cpp
+++ b/Engine/Source/Editor/UnrealEd/Private/Cooker/CookSandbox.cpp
@@ -8,6 +8,7 @@
 #include "IPlatformFileSandboxWrapper.h"
 #include "Misc/App.h"
 #include "Misc/CString.h"
+#include "Misc/PackageName.h"
 #include "Misc/Paths.h"
 #include "String/Find.h"
 
@@ -296,4 +297,4 @@ FString FCookSandbox::ConvertPackageNameToCookedPath(FStringView PackageName,
 }
 
 
-} // namespace UE::Cook
\ No newline at end of file
+} // namespace UE::Cook
