Allow GFP Policy to determine if a GFP can allow async load

diff --git a/Engine/Plugins/Runtime/GameFeatures/Source/GameFeatures/Private/GameFeaturePluginStateMachine.cpp b/Engine/Plugins/Runtime/GameFeatures/Source/GameFeatures/Private/GameFeaturePluginStateMachine.cpp
--- a/Engine/Plugins/Runtime/GameFeatures/Source/GameFeatures/Private/GameFeaturePluginStateMachine.cpp
+++ b/Engine/Plugins/Runtime/GameFeatures/Source/GameFeatures/Private/GameFeaturePluginStateMachine.cpp
@@ -80,10 +80,6 @@ namespace UE::GameFeatures
 		true,
 		TEXT("Enable to use async loading as well async downloading and registering"));
 
-	static TAutoConsoleVariable<bool> CVarForceAsyncLoad(TEXT("GameFeaturePlugin.ForceAsyncLoad"),
-		false,
-		TEXT("Enable to force use of async loading even if normally not allowed"));
-
 	static TAutoConsoleVariable<bool> CVarAllowForceMonolithicShaderLibrary(TEXT("GameFeaturePlugin.AllowForceMonolithicShaderLibrary"),
 		true,
 		TEXT("Enable to force only searching for monolithic shader libs when possible"));
@@ -4877,7 +4873,7 @@ bool FGameFeaturePluginStateMachineProperties::AllowAsyncLoading() const
 	// The local bForceSyncLoading should take precedence over UE::GameFeatures::CVarForceAsyncLoad
 	return
 		!ProtocolOptions.bForceSyncLoading &&
-		(!IsRunningCommandlet() || UE::GameFeatures::CVarForceAsyncLoad.GetValueOnGameThread());
+		UGameFeaturesSubsystem::Get().GetPolicy().AllowAsyncLoad(PluginIdentifier.GetFullPluginURL());
 }
 
 bool FGameFeaturePluginStateMachineProperties::CanBatchProcess() const
diff --git a/Engine/Plugins/Runtime/GameFeatures/Source/GameFeatures/Private/GameFeaturesProjectPolicies.cpp b/Engine/Plugins/Runtime/GameFeatures/Source/GameFeatures/Private/GameFeaturesProjectPolicies.cpp
--- a/Engine/Plugins/Runtime/GameFeatures/Source/GameFeatures/Private/GameFeaturesProjectPolicies.cpp
+++ b/Engine/Plugins/Runtime/GameFeatures/Source/GameFeatures/Private/GameFeaturesProjectPolicies.cpp
@@ -8,6 +8,18 @@
 
 #include UE_INLINE_GENERATED_CPP_BY_NAME(GameFeaturesProjectPolicies)
 
+namespace UE::GameFeatures
+{
+	static TAutoConsoleVariable<bool> CVarForceAsyncLoad(TEXT("GameFeaturePlugin.ForceAsyncLoad"),
+		false,
+		TEXT("Enable to force use of async loading even if normally not allowed"));
+
+	const TAutoConsoleVariable<bool>& GetCVarForceAsyncLoad()
+	{
+		return CVarForceAsyncLoad;
+	}
+}
+
 void UDefaultGameFeaturesProjectPolicies::InitGameFeatureManager()
 {
 	UE_LOG(LogGameFeatures, Log, TEXT("Scanning for built-in game feature plugins"));
@@ -146,3 +158,8 @@ void UGameFeaturesProjectPolicies::ExplicitLoadGameFeaturePlugin(const FString&
 		UGameFeaturesSubsystem::Get().LoadGameFeaturePlugin(PluginURL, CompleteDelegate);
 	}
 }
+
+bool UGameFeaturesProjectPolicies::AllowAsyncLoad(FStringView /*PluginURL*/) const
+{
+	return !IsRunningCommandlet() || UE::GameFeatures::CVarForceAsyncLoad.GetValueOnGameThread();
+}
diff --git a/Engine/Plugins/Runtime/GameFeatures/Source/GameFeatures/Public/GameFeaturesProjectPolicies.h b/Engine/Plugins/Runtime/GameFeatures/Source/GameFeatures/Public/GameFeaturesProjectPolicies.h
--- a/Engine/Plugins/Runtime/GameFeatures/Source/GameFeatures/Public/GameFeaturesProjectPolicies.h
+++ b/Engine/Plugins/Runtime/GameFeatures/Source/GameFeatures/Public/GameFeaturesProjectPolicies.h
@@ -11,6 +11,11 @@
 
 class UGameFeatureData;
 
+namespace UE::GameFeatures
+{
+	UE_API extern const TAutoConsoleVariable<bool>& GetCVarForceAsyncLoad();
+}
+
 struct FPluginDependencyDetails
 {
 	bool bFailIfNotFound = false;
@@ -92,6 +97,9 @@ public:
 	// Called by code that explicitly wants to load a specific plugin
 	// (e.g., when using a fast cook a game might want to allow explicitly loaded game feature plugins)
 	UE_API virtual void ExplicitLoadGameFeaturePlugin(const FString& PluginURL, const FGameFeaturePluginLoadComplete& CompleteDelegate, const bool bActivateGameFeatures);
+
+	// Called to determine if async loading is allowed
+	UE_API virtual bool AllowAsyncLoad(FStringView PluginURL) const;
 };
 
 // This is a default implementation that immediately processes all game feature plugins the based on
