Add option for GameFeatureSubsystem observers to not be updated for current state when added. Allows user to avoid perf cost of iterating plugins when not needed.

diff --git a/Engine/Plugins/Experimental/CommonConversation/Source/CommonConversationRuntime/Private/ConversationRegistry.cpp b/Engine/Plugins/Experimental/CommonConversation/Source/CommonConversationRuntime/Private/ConversationRegistry.cpp
--- a/Engine/Plugins/Experimental/CommonConversation/Source/CommonConversationRuntime/Private/ConversationRegistry.cpp
+++ b/Engine/Plugins/Experimental/CommonConversation/Source/CommonConversationRuntime/Private/ConversationRegistry.cpp
@@ -162,7 +162,8 @@ void UConversationRegistry::Initialize(FSubsystemCollectionBase& Collection)
 {
 	Super::Initialize(Collection);
 
-	UGameFeaturesSubsystem::Get().AddObserver(this);
+	UGameFeaturesSubsystem::Get().AddObserver(this, UGameFeaturesSubsystem::EObserverPluginStateUpdateMode::FutureOnly);
+	GameFeatureStateModified();
 }
 
 void UConversationRegistry::Deinitialize()
diff --git a/Engine/Plugins/Runtime/GameFeatures/Source/GameFeatures/Private/GameFeaturesSubsystem.cpp b/Engine/Plugins/Runtime/GameFeatures/Source/GameFeatures/Private/GameFeaturesSubsystem.cpp
--- a/Engine/Plugins/Runtime/GameFeatures/Source/GameFeatures/Private/GameFeaturesSubsystem.cpp
+++ b/Engine/Plugins/Runtime/GameFeatures/Source/GameFeatures/Private/GameFeaturesSubsystem.cpp
@@ -835,33 +835,44 @@ void UGameFeaturesSubsystem::ForEachGameFeature(TFunctionRef<void(FGameFeatureIn
 
 void UGameFeaturesSubsystem::AddObserver(UObject* Observer)
 {
+	//deprecated, previous default behaviour
+	AddObserver(Observer, EObserverPluginStateUpdateMode::CurrentAndFuture);
+}
+
+void UGameFeaturesSubsystem::AddObserver(UObject* Observer, EObserverPluginStateUpdateMode UpdateMode)
+{
+	QUICK_SCOPE_CYCLE_COUNTER(UGameFeaturesSubsystem_AddObserver);
+
 	check(Observer);
 	IGameFeatureStateChangeObserver* Interface = Cast<IGameFeatureStateChangeObserver>(Observer);
 	if (ensureAlwaysMsgf(Interface != nullptr, TEXT("Observers must implement the IGameFeatureStateChangeObserver interface.")))
 	{
 		Observers.AddUnique(Observer);
 
-		// Push the current state of all known game features to the new observer
-		for (auto StateMachineIt = GameFeaturePluginStateMachines.CreateConstIterator(); StateMachineIt; ++StateMachineIt)
+		if (UpdateMode == EObserverPluginStateUpdateMode::CurrentAndFuture)
 		{
-			if (UGameFeaturePluginStateMachine* GFSM = StateMachineIt.Value())
+			// Push the current state of all known game features to the new observer
+			for (auto StateMachineIt = GameFeaturePluginStateMachines.CreateConstIterator(); StateMachineIt; ++StateMachineIt)
 			{
-				if (const UGameFeatureData* GameFeatureData = GFSM->GetGameFeatureDataForRegisteredPlugin(false))
+				if (UGameFeaturePluginStateMachine* GFSM = StateMachineIt.Value())
 				{
-					FString PluginName = GFSM->GetPluginName();
-					FString PluginURL = GFSM->GetPluginURL();
+					if (const UGameFeatureData* GameFeatureData = GFSM->GetGameFeatureDataForRegisteredPlugin(false))
+					{
+						const FString& PluginName = GFSM->GetPluginName();
+						const FString& PluginURL = GFSM->GetPluginURL();
 
-					Interface->OnGameFeatureRegistering(GameFeatureData, *PluginName, PluginURL);
+						Interface->OnGameFeatureRegistering(GameFeatureData, PluginName, PluginURL);
 
-					if (GFSM->GetCurrentState() >= EGameFeaturePluginState::Loaded)
-					{
-						Interface->OnGameFeatureLoading(GameFeatureData, PluginURL);
-					}
+						if (GFSM->GetCurrentState() >= EGameFeaturePluginState::Loaded)
+						{
+							Interface->OnGameFeatureLoading(GameFeatureData, PluginURL);
+						}
 
-					if (GFSM->GetCurrentState() >= EGameFeaturePluginState::Active)
-					{
-						Interface->OnGameFeatureActivating(GameFeatureData, PluginURL);
-						Interface->OnGameFeatureActivated(GameFeatureData, PluginURL);
+						if (GFSM->GetCurrentState() >= EGameFeaturePluginState::Active)
+						{
+							Interface->OnGameFeatureActivating(GameFeatureData, PluginURL);
+							Interface->OnGameFeatureActivated(GameFeatureData, PluginURL);
+						}
 					}
 				}
 			}
diff --git a/Engine/Plugins/Runtime/GameFeatures/Source/GameFeatures/Public/GameFeaturesSubsystem.h b/Engine/Plugins/Runtime/GameFeatures/Source/GameFeatures/Public/GameFeaturesSubsystem.h
--- a/Engine/Plugins/Runtime/GameFeatures/Source/GameFeatures/Public/GameFeaturesSubsystem.h
+++ b/Engine/Plugins/Runtime/GameFeatures/Source/GameFeatures/Public/GameFeaturesSubsystem.h
@@ -448,7 +448,17 @@ public:
 	static UE_API TSharedPtr<FStreamableHandle> LoadGameFeatureData(const FString& GameFeatureToLoad, bool bStartStalled = false);
 	static UE_API void UnloadGameFeatureData(const UGameFeatureData* GameFeatureToUnload);
 
+	UE_DEPRECATED(5.7, "Use AddObserver with EObserverPluginStateUpdateMode parameter")
 	UE_API void AddObserver(UObject* Observer);
+
+	enum class EObserverPluginStateUpdateMode
+	{
+		FutureOnly, //Only call the observer with future plugin state changes
+		CurrentAndFuture, //Also update the observer with the current plugin state at add time - note, if project has a lot of plugins this can be expensive
+	};
+
+	//Add Observer to be notified when plugin states change. Observer must implement IGameFeatureStateChangeObserver and be non-null
+	UE_API void AddObserver(UObject* Observer, const EObserverPluginStateUpdateMode UpdateMode);
 	UE_API void RemoveObserver(UObject* Observer);
 
 	UE_API void ForEachGameFeature(TFunctionRef<void(FGameFeatureInfo&&)> Visitor) const;
