Workaround for UE-309330
- Do not rely on FPackageName::DoesPackageExist when running with cook-on-the-fly, it is unreliable

#tests Tested Lyra with and without cook on the fly, async loading

diff --git a/Engine/Plugins/Runtime/GameFeatures/Source/GameFeatures/Private/GameFeaturePluginStateMachine.cpp b/Engine/Plugins/Runtime/GameFeatures/Source/GameFeatures/Private/GameFeaturePluginStateMachine.cpp
--- a/Engine/Plugins/Runtime/GameFeatures/Source/GameFeatures/Private/GameFeaturePluginStateMachine.cpp
+++ b/Engine/Plugins/Runtime/GameFeatures/Source/GameFeatures/Private/GameFeaturePluginStateMachine.cpp
@@ -3147,16 +3147,24 @@ struct FGameFeaturePluginState_Registering : public FGameFeaturePluginState
 	TSharedPtr<FStreamableHandle> GameFeatureDataHandle;
 	ELoadGFDState LoadGFDState = ELoadGFDState::Pending;
 	bool bCheckedRealtimeMode = false;
+	TArray<FString, TInlineAllocator<2>> GameFeatureDataPaths;
 
 	FGameFeaturePluginState_Registering(FGameFeaturePluginStateMachineProperties& InStateProperties) : FGameFeaturePluginState(InStateProperties) {}
 	
-	void TryAsyncLoadGameFeatureData(const FString& Path)
+	void TryAsyncLoadGameFeatureData(int32 Attempt = 0)
 	{
 		check(LoadGFDState == ELoadGFDState::Pending);
 
+		if (!GameFeatureDataPaths.IsValidIndex(Attempt))
+		{
+			LoadGFDState = ELoadGFDState::Failed;
+			UpdateStateMachineDeferred();
+			return;
+		}
+
 		bool bIsLoading = false;
 
-		GameFeatureDataHandle = UGameFeaturesSubsystem::LoadGameFeatureData(Path, true /*bStartStalled*/);
+		GameFeatureDataHandle = UGameFeaturesSubsystem::LoadGameFeatureData(GameFeatureDataPaths[Attempt], true /*bStartStalled*/);
 		if (GameFeatureDataHandle && GameFeatureDataHandle->IsLoadingInProgress())
 		{
 			// CurrentMachine owns `this`, so use it as a lifetime check for the `this` captured in the lambdas below, as they may be called after `this` is destroyed
@@ -3177,7 +3185,7 @@ struct FGameFeaturePluginState_Registering : public FGameFeaturePluginState
 				UpdateStateMachineDeferred();
 			}));
 
-			GameFeatureDataHandle->BindCompleteDelegate(FStreamableDelegateWithHandle::CreateWeakLambda(CurrentMachine, [this](const TSharedPtr<FStreamableHandle>& Handle)
+			GameFeatureDataHandle->BindCompleteDelegate(FStreamableDelegateWithHandle::CreateWeakLambda(CurrentMachine, [this, Attempt](const TSharedPtr<FStreamableHandle>& Handle)
 			{
 				if (GameFeatureDataHandle != Handle)
 				{
@@ -3186,6 +3194,11 @@ struct FGameFeaturePluginState_Registering : public FGameFeaturePluginState
 				}
 
 				StateProperties.GameFeatureData = Cast<UGameFeatureData>(GameFeatureDataHandle->GetLoadedAsset());
+				if (!StateProperties.GameFeatureData && GameFeatureDataPaths.IsValidIndex(Attempt + 1))
+				{
+					TryAsyncLoadGameFeatureData(Attempt + 1);
+					return;
+				}
 
 				if (StateProperties.GameFeatureData)
 				{
@@ -3215,7 +3228,7 @@ struct FGameFeaturePluginState_Registering : public FGameFeaturePluginState
 
 		if (!bIsLoading)
 		{
-			LoadGFDState = ELoadGFDState::Failed;
+			TryAsyncLoadGameFeatureData(Attempt + 1);
 		}
 	}
 
@@ -3282,27 +3295,37 @@ struct FGameFeaturePluginState_Registering : public FGameFeaturePluginState
 			}
 		}
 
-		FString GameFeatureDataPath;
-		if (FPackageName::DoesPackageExist(PreferredGameFeatureDataPath))
-		{
-			GameFeatureDataPath = MoveTemp(PreferredGameFeatureDataPath);
-		}
-		else if (FPackageName::DoesPackageExist(BackupGameFeatureDataPath))
+		// Temporary workaround for UE-309330
+		// FPackageName::DoesPackageExist is unreliable when running with cook-on-the-fly, instead we must make multiple attempts
+		// to load the cooked packages and handle failure/repeats
+		if (IsRunningCookOnTheFly())
 		{
-			GameFeatureDataPath = MoveTemp(BackupGameFeatureDataPath);
+			GameFeatureDataPaths.Emplace(MoveTemp(PreferredGameFeatureDataPath));
+			GameFeatureDataPaths.Emplace(MoveTemp(BackupGameFeatureDataPath));
 		}
 		else
 		{
-			const FStringView ShortUrl = StateProperties.PluginIdentifier.GetIdentifyingString();
-			UE_LOG(LogGameFeatures, Error, TEXT("Game Feature Data package not found for URL %.*s"), ShortUrl.Len(), ShortUrl.GetData());
+			if (FPackageName::DoesPackageExist(PreferredGameFeatureDataPath))
+			{
+				GameFeatureDataPaths.Emplace(MoveTemp(PreferredGameFeatureDataPath));
+			}
+			else if (FPackageName::DoesPackageExist(BackupGameFeatureDataPath))
+			{
+				GameFeatureDataPaths.Emplace(MoveTemp(BackupGameFeatureDataPath));
+			}
+			else
+			{
+				const FStringView ShortUrl = StateProperties.PluginIdentifier.GetIdentifyingString();
+				UE_LOG(LogGameFeatures, Error, TEXT("Game Feature Data package not found for URL %.*s"), ShortUrl.Len(), ShortUrl.GetData());
 
-			LoadGFDState = ELoadGFDState::Failed;
-			return;
+				LoadGFDState = ELoadGFDState::Failed;
+				return;
+			}
 		}
 
 		if (UseAsyncLoading())
 		{
-			TryAsyncLoadGameFeatureData(GameFeatureDataPath);
+			TryAsyncLoadGameFeatureData();
 		}
 		else
 		{
@@ -3312,11 +3335,18 @@ struct FGameFeaturePluginState_Registering : public FGameFeaturePluginState
 					FText::FromString(StateProperties.PluginName)));
 			LoadingGameFeatureData.Visibility = ESlowTaskVisibility::Important;
 			
-			GameFeatureDataHandle = UGameFeaturesSubsystem::LoadGameFeatureData(GameFeatureDataPath);
-			if (GameFeatureDataHandle)
+			for (const FString& Path : GameFeatureDataPaths)
 			{
-				GameFeatureDataHandle->WaitUntilComplete(0.0f, false);
-				StateProperties.GameFeatureData = Cast<UGameFeatureData>(GameFeatureDataHandle->GetLoadedAsset());
+				GameFeatureDataHandle = UGameFeaturesSubsystem::LoadGameFeatureData(Path);
+				if (GameFeatureDataHandle)
+				{
+					GameFeatureDataHandle->WaitUntilComplete(0.0f, false);
+					StateProperties.GameFeatureData = Cast<UGameFeatureData>(GameFeatureDataHandle->GetLoadedAsset());
+				}
+				if (StateProperties.GameFeatureData)
+				{
+					break;
+				}
 			}
 
 			if (StateProperties.GameFeatureData)
