Prevent world subsystems from initializing when their owning plugin is not active

#rnx
#changelist validated
#virtualized

diff --git a/Engine/Plugins/Runtime/GameFeatures/Source/GameFeatures/Private/GameFeaturesSubsystem.cpp b/Engine/Plugins/Runtime/GameFeatures/Source/GameFeatures/Private/GameFeaturesSubsystem.cpp
--- a/Engine/Plugins/Runtime/GameFeatures/Source/GameFeatures/Private/GameFeaturesSubsystem.cpp
+++ b/Engine/Plugins/Runtime/GameFeatures/Source/GameFeatures/Private/GameFeaturesSubsystem.cpp
@@ -1970,6 +1970,25 @@ bool UGameFeaturesSubsystem::IsGameFeaturePluginActive(const FString& PluginURL,
 	return false;
 }
 
+bool UGameFeaturesSubsystem::IsGameFeaturePluginActiveByName(FStringView PluginName, bool bCheckForActivating /*= false*/) const
+{
+	FString PluginURLName;
+	if (GetPluginURLByName(PluginName, PluginURLName))
+	{
+		if (!IsGameFeaturePluginActive(PluginURLName, bCheckForActivating))
+		{
+			return false;
+		}
+	}
+	else
+	{
+		UE_LOG(LogGameFeatures, Error, TEXT("UGameFeaturesSubsystem could not get the plugin URL from the plugin name. Name:%.*s"), PluginName.Len(), PluginName.GetData());
+		return false;
+	}
+
+	return true;
+}
+
 void UGameFeaturesSubsystem::DeactivateGameFeaturePlugin(const FString& PluginURL)
 {
 	DeactivateGameFeaturePlugin(PluginURL, FGameFeaturePluginDeactivateComplete());
diff --git a/Engine/Plugins/Runtime/GameFeatures/Source/GameFeatures/Public/GameFeaturesSubsystem.h b/Engine/Plugins/Runtime/GameFeatures/Source/GameFeatures/Public/GameFeaturesSubsystem.h
--- a/Engine/Plugins/Runtime/GameFeatures/Source/GameFeatures/Public/GameFeaturesSubsystem.h
+++ b/Engine/Plugins/Runtime/GameFeatures/Source/GameFeatures/Public/GameFeaturesSubsystem.h
@@ -588,6 +588,9 @@ public:
 	/** Determines if a plugin is in the Active state.*/
 	UE_API bool IsGameFeaturePluginActive(const FString& PluginURL, bool bCheckForActivating = false) const;
 
+	/** Determines if a plugin is in the Active state.*/
+	UE_API bool IsGameFeaturePluginActiveByName(FStringView PluginName, bool bCheckForActivating = false) const;
+
 	/** Determines if a plugin is up to date or needs an update. Returns true if an update is available.*/
 	UE_API bool DoesGameFeaturePluginNeedUpdate(const FString& PluginURL) const;
 
