GFP State machine - Implement DoNotDownload flag for IAD content.

diff --git a/Engine/Plugins/Runtime/GameFeatures/Source/GameFeatures/Private/GameFeaturePluginStateMachine.cpp b/Engine/Plugins/Runtime/GameFeatures/Source/GameFeatures/Private/GameFeaturePluginStateMachine.cpp
--- a/Engine/Plugins/Runtime/GameFeatures/Source/GameFeatures/Private/GameFeaturePluginStateMachine.cpp
+++ b/Engine/Plugins/Runtime/GameFeatures/Source/GameFeatures/Private/GameFeaturePluginStateMachine.cpp
@@ -1748,6 +1748,8 @@ struct FGameFeaturePluginState_Downloading : public FGameFeaturePluginState
 			return;
 		}
 
+		const FInstallBundlePluginProtocolOptions& Options = StateProperties.ProtocolOptions.GetSubtype<FInstallBundlePluginProtocolOptions>();
+
 		IoStoreOnDemandContext = MakeShared<FIoStoreOnDemandContext>();
 		IoStoreOnDemandContext->PendingInstalls = MetaData.InstallBundlesWithAssetDependencies.Num();
 
@@ -1758,6 +1760,10 @@ struct FGameFeaturePluginState_Downloading : public FGameFeaturePluginState
 			InstallArgs.TagSets.Emplace(TEXTVIEW("mount")); // May not be a real tagset, but this will install all the mandatory untagged chunks
 			InstallArgs.Options |= UE::IoStore::EOnDemandInstallOptions::InstallSoftReferences;
 			InstallArgs.Options |= UE::IoStore::EOnDemandInstallOptions::CallbackOnGameThread;
+			if (Options.bDoNotDownload)
+			{
+				InstallArgs.Options |= UE::IoStore::EOnDemandInstallOptions::DoNotDownload;
+			}
 			
 			InstallArgs.ContentHandle = UE::GameFeatures::GFPSharedInstallTracker.AddOnDemandContentHandle(
 				InstallBundle, UE::GameFeatures::EGFPInstallLevel::Download);
@@ -2949,6 +2955,8 @@ struct FGameFeaturePluginState_AssetDependencyStreaming : public FGameFeaturePlu
 			return;
 		}
 
+		const FInstallBundlePluginProtocolOptions& Options = StateProperties.ProtocolOptions.GetSubtype<FInstallBundlePluginProtocolOptions>();
+
 		UE::GameFeatures::GFPSharedInstallTracker.AddBundleRefs(
 			StateProperties.PluginName, UE::GameFeatures::EGFPInstallLevel::AssetStream, MetaData.InstallBundlesWithAssetDependencies);
 
@@ -2971,6 +2979,10 @@ struct FGameFeaturePluginState_AssetDependencyStreaming : public FGameFeaturePlu
 			}
 			InstallArgs.Options |= UE::IoStore::EOnDemandInstallOptions::InstallSoftReferences;
 			InstallArgs.Options |= UE::IoStore::EOnDemandInstallOptions::CallbackOnGameThread;
+			if (Options.bDoNotDownload)
+			{
+				InstallArgs.Options |= UE::IoStore::EOnDemandInstallOptions::DoNotDownload;
+			}
 
 			InstallArgs.ContentHandle = UE::GameFeatures::GFPSharedInstallTracker.AddOnDemandContentHandle(
 				InstallBundle, UE::GameFeatures::EGFPInstallLevel::AssetStream);
diff --git a/Engine/Source/Runtime/Experimental/IoStore/OnDemand/Private/OnDemandContentInstaller.cpp b/Engine/Source/Runtime/Experimental/IoStore/OnDemand/Private/OnDemandContentInstaller.cpp
--- a/Engine/Source/Runtime/Experimental/IoStore/OnDemand/Private/OnDemandContentInstaller.cpp
+++ b/Engine/Source/Runtime/Experimental/IoStore/OnDemand/Private/OnDemandContentInstaller.cpp
@@ -709,6 +709,13 @@ void FOnDemandContentInstaller::ExecuteInstallRequest(FRequest& Request)
 		return CompleteInstallRequest(Request);
 	}
 
+	if (EnumHasAnyFlags(InstallRequest.Args.Options, EOnDemandInstallOptions::DoNotDownload))
+	{
+		Request.ErrorCode = EIoErrorCode::Disabled;
+		Request.ErrorReason = TEXT("DoNotDownload Set");
+		return CompleteInstallRequest(Request);
+	}
+
 	// Make sure we have enough space in the cache
 	{
 		TSet<FIoHash> ChunksToInstall;
diff --git a/Engine/Source/Runtime/Experimental/IoStore/OnDemandCore/Public/IO/IoStoreOnDemand.h b/Engine/Source/Runtime/Experimental/IoStore/OnDemandCore/Public/IO/IoStoreOnDemand.h
--- a/Engine/Source/Runtime/Experimental/IoStore/OnDemandCore/Public/IO/IoStoreOnDemand.h
+++ b/Engine/Source/Runtime/Experimental/IoStore/OnDemandCore/Public/IO/IoStoreOnDemand.h
@@ -187,6 +187,8 @@ enum class EOnDemandInstallOptions
 	InstallSoftReferences	= 1 << 1,
 	/** Install optional bulk data. */
 	InstallOptionalBulkData = 1 << 2,
+	/** Do not install any data, only succeed if all data is already cached. */
+	DoNotDownload			= 1 << 3,
 };
 ENUM_CLASS_FLAGS(EOnDemandInstallOptions);
 
