Add CVar to disable GFP errors when IAD depdencies are missing.  This can be used to mitigate problems that may arise if a pakchunk with an asset isn't mounted.  Instead of failing the install, the asset will just fail to load at runtime.

diff --git a/Engine/Plugins/Runtime/GameFeatures/Source/GameFeatures/Private/GameFeaturePluginStateMachine.cpp b/Engine/Plugins/Runtime/GameFeatures/Source/GameFeatures/Private/GameFeaturePluginStateMachine.cpp
--- a/Engine/Plugins/Runtime/GameFeatures/Source/GameFeatures/Private/GameFeaturePluginStateMachine.cpp
+++ b/Engine/Plugins/Runtime/GameFeatures/Source/GameFeatures/Private/GameFeaturePluginStateMachine.cpp
@@ -108,6 +108,10 @@ namespace UE::GameFeatures
 		false,
 		TEXT("Enable batch processing when processing multiple plugins and specified in protocol options."));
 
+	/*extern*/ TAutoConsoleVariable<bool> CVarAllowMissingOnDemandDependencies(TEXT("GameFeaturePlugin.AllowMissingOnDemandDependencies"),
+		false,
+		TEXT("Don't fail on missing on demand (IAD) asset dependencies"));
+
 	bool ShouldDeferLocalizationDataLoad()
 	{
 		// Note: We don't defer localization data loading in the editor, as the editor only needs to mount plugins to use them
@@ -2983,6 +2987,10 @@ struct FGameFeaturePluginState_AssetDependencyStreaming : public FGameFeaturePlu
 			{
 				InstallArgs.Options |= UE::IoStore::EOnDemandInstallOptions::DoNotDownload;
 			}
+			if (UE::GameFeatures::CVarAllowMissingOnDemandDependencies.GetValueOnGameThread())
+			{
+				InstallArgs.Options |= UE::IoStore::EOnDemandInstallOptions::AllowMissingDependencies;
+			}
 
 			InstallArgs.ContentHandle = UE::GameFeatures::GFPSharedInstallTracker.AddOnDemandContentHandle(
 				InstallBundle, UE::GameFeatures::EGFPInstallLevel::AssetStream);
diff --git a/Engine/Plugins/Runtime/GameFeatures/Source/GameFeatures/Private/GameFeaturePluginStateMachine.h b/Engine/Plugins/Runtime/GameFeatures/Source/GameFeatures/Private/GameFeaturePluginStateMachine.h
--- a/Engine/Plugins/Runtime/GameFeatures/Source/GameFeatures/Private/GameFeaturePluginStateMachine.h
+++ b/Engine/Plugins/Runtime/GameFeatures/Source/GameFeatures/Private/GameFeaturePluginStateMachine.h
@@ -15,6 +15,11 @@ struct FComponentRequestHandle;
 enum class EInstallBundleResult : uint32;
 enum class EInstallBundleReleaseResult : uint32;
 
+namespace UE::GameFeatures
+{
+	extern TAutoConsoleVariable<bool> CVarAllowMissingOnDemandDependencies;
+}
+
 /*
 *************** GameFeaturePlugin state machine graph ***************
 Descriptions for each state are below in EGameFeaturePluginState.
diff --git a/Engine/Plugins/Runtime/GameFeatures/Source/GameFeatures/Private/GameFeaturesSubsystem.cpp b/Engine/Plugins/Runtime/GameFeatures/Source/GameFeatures/Private/GameFeaturesSubsystem.cpp
--- a/Engine/Plugins/Runtime/GameFeatures/Source/GameFeatures/Private/GameFeaturesSubsystem.cpp
+++ b/Engine/Plugins/Runtime/GameFeatures/Source/GameFeatures/Private/GameFeaturesSubsystem.cpp
@@ -3435,6 +3435,10 @@ struct FGameFeaturePluginPredownloadContext : public FGameFeaturePluginPredownlo
 				InstallArgs.TagSets.Emplace(TEXTVIEW("required"));
 			}
 			InstallArgs.Options |= UE::IoStore::EOnDemandInstallOptions::InstallSoftReferences;
+			if (UE::GameFeatures::CVarAllowMissingOnDemandDependencies.GetValueOnGameThread())
+			{
+				InstallArgs.Options |= UE::IoStore::EOnDemandInstallOptions::AllowMissingDependencies;
+			}
 
 			InstallArgs.ContentHandle = IoStoreContentHandles[i];
 
diff --git a/Engine/Source/Runtime/Experimental/IoStore/OnDemand/Private/OnDemandContentInstaller.cpp b/Engine/Source/Runtime/Experimental/IoStore/OnDemand/Private/OnDemandContentInstaller.cpp
--- a/Engine/Source/Runtime/Experimental/IoStore/OnDemand/Private/OnDemandContentInstaller.cpp
+++ b/Engine/Source/Runtime/Experimental/IoStore/OnDemand/Private/OnDemandContentInstaller.cpp
@@ -587,7 +587,7 @@ void FOnDemandContentInstaller::ProcessInstallRequest(FRequest& Request)
 			}
 		}
 
-		if (MissingCount > 0) 
+		if (MissingCount > 0 && !EnumHasAnyFlags(InstallRequest.Args.Options, EOnDemandInstallOptions::AllowMissingDependencies))
 		{
 			Request.ErrorCode	= EIoErrorCode::UnknownChunkID;
 			Request.ErrorReason	= FString::Printf(TEXT("Missing chunk(s), Count=%u, ContentHandle='%s'"),
diff --git a/Engine/Source/Runtime/Experimental/IoStore/OnDemandCore/Public/IO/IoStoreOnDemand.h b/Engine/Source/Runtime/Experimental/IoStore/OnDemandCore/Public/IO/IoStoreOnDemand.h
--- a/Engine/Source/Runtime/Experimental/IoStore/OnDemandCore/Public/IO/IoStoreOnDemand.h
+++ b/Engine/Source/Runtime/Experimental/IoStore/OnDemandCore/Public/IO/IoStoreOnDemand.h
@@ -180,15 +180,17 @@ using FOnDemandMountCompleted = TUniqueFunction<void(FOnDemandMountResult)>;
 enum class EOnDemandInstallOptions
 {
 	/** No additional options. */
-	None					= 0,
+	None						= 0,
 	/** Trigger callback on game thread. */
-	CallbackOnGameThread	= 1 << 0,
+	CallbackOnGameThread		= 1 << 0,
 	/** Follow soft references when gathering packages to install. */
-	InstallSoftReferences	= 1 << 1,
+	InstallSoftReferences		= 1 << 1,
 	/** Install optional bulk data. */
-	InstallOptionalBulkData = 1 << 2,
+	InstallOptionalBulkData		= 1 << 2,
 	/** Do not install any data, only succeed if all data is already cached. */
-	DoNotDownload			= 1 << 3,
+	DoNotDownload				= 1 << 3,
+	/** Do not return an error if missing iochunk dependencies are missing. */
+	AllowMissingDependencies	= 1 << 4
 };
 ENUM_CLASS_FLAGS(EOnDemandInstallOptions);
 
