[GFP] Do not release dependencies of optional content, FGFPSharedInstallTracker already handles dependencies.

#rnx

diff --git a/Engine/Plugins/Runtime/GameFeatures/Source/GameFeatures/Private/GameFeatureOptionalContentInstaller.cpp b/Engine/Plugins/Runtime/GameFeatures/Source/GameFeatures/Private/GameFeatureOptionalContentInstaller.cpp
--- a/Engine/Plugins/Runtime/GameFeatures/Source/GameFeatures/Private/GameFeatureOptionalContentInstaller.cpp
+++ b/Engine/Plugins/Runtime/GameFeatures/Source/GameFeatures/Private/GameFeatureOptionalContentInstaller.cpp
@@ -295,6 +295,8 @@ void UGameFeatureOptionalContentInstaller::ReleaseContent(const FString& PluginN
 		return;
 	}
 
+	Flags |= EInstallBundleReleaseRequestFlags::ExplicitRemoveList;
+
 	BundleManager->RequestReleaseContent(
 		Bundles, 
 		Flags,
