[InstallBundle] Allow to update only requested bundles. Forbid load dependencies for OptionalContentInstaller

#rnx

diff --git a/Engine/Plugins/Experimental/DefaultInstallBundleManager/Source/Private/DefaultInstallBundleManager.cpp b/Engine/Plugins/Experimental/DefaultInstallBundleManager/Source/Private/DefaultInstallBundleManager.cpp
--- a/Engine/Plugins/Experimental/DefaultInstallBundleManager/Source/Private/DefaultInstallBundleManager.cpp
+++ b/Engine/Plugins/Experimental/DefaultInstallBundleManager/Source/Private/DefaultInstallBundleManager.cpp
@@ -11,6 +11,7 @@
 #include "Misc/ConfigContext.h"
 #include "Misc/CommandLine.h"
 #include "IO/IoStoreOnDemand.h"
+#include "Misc/EnumClassFlags.h"
 #include "PatchCheck.h"
 #include "Templates/Greater.h"
 #include "ShaderPipelineCache.h"
@@ -3430,7 +3431,15 @@ TValueOrError<FInstallBundleRequestInfo, EInstallBundleResult> FDefaultInstallBu
 		return MakeError(EInstallBundleResult::InitializationPending);
 	}
 
-	TSet<FName> BundlesToLoad = GatherBundlesForRequest(InBundleNames, RetInfo.InfoFlags);
+	TSet<FName> BundlesToLoad;
+	if (EnumHasAnyFlags(Flags, EInstallBundleRequestFlags::ExplicitUpdateList))
+	{
+		BundlesToLoad.Append(InBundleNames);
+	}
+	else
+	{
+		BundlesToLoad = GatherBundlesForRequest(InBundleNames, RetInfo.InfoFlags);
+	}
 	for (const FName& BundleName : BundlesToLoad)
 	{
 		const FBundleInfo* BundleInfo = BundleInfoMap.Find(BundleName);
diff --git a/Engine/Plugins/Runtime/GameFeatures/Source/GameFeatures/Private/GameFeatureOptionalContentInstaller.cpp b/Engine/Plugins/Runtime/GameFeatures/Source/GameFeatures/Private/GameFeatureOptionalContentInstaller.cpp
--- a/Engine/Plugins/Runtime/GameFeatures/Source/GameFeatures/Private/GameFeatureOptionalContentInstaller.cpp
+++ b/Engine/Plugins/Runtime/GameFeatures/Source/GameFeatures/Private/GameFeatureOptionalContentInstaller.cpp
@@ -123,11 +123,14 @@ bool UGameFeatureOptionalContentInstaller::UpdateContent(const FString& PluginNa
 	bool bIsAvailable = false;
 	if (!Bundles.IsEmpty())
 	{
-		TValueOrError<FInstallBundleCombinedInstallState, EInstallBundleResult> MaybeInstallState = BundleManager->GetInstallStateSynchronous(Bundles, false);
+		TValueOrError<FInstallBundleCombinedInstallState, EInstallBundleResult> MaybeInstallState = BundleManager->GetInstallStateSynchronous(Bundles, true);
 		if (MaybeInstallState.HasValue())
 		{
 			const FInstallBundleCombinedInstallState& InstallState = MaybeInstallState.GetValue();
 			bIsAvailable = Algo::AllOf(Bundles, [&InstallState](FName BundleName) { return InstallState.IndividualBundleStates.Contains(BundleName); });
+			bIsAvailable &= ensureMsgf(InstallState.IndividualBundleStates.Num() <= Bundles.Num(),
+				TEXT("UGameFeatureOptionalContentInstaller does not support dependencies tracking. Check optional install bundle dependencies for plugin %s."),
+				*PluginName);
 		}
 	}
 
@@ -140,7 +143,7 @@ bool UGameFeatureOptionalContentInstaller::UpdateContent(const FString& PluginNa
 	{
 		UE_LOG(LogGameFeatures, Log, TEXT("Requesting update for %s"), *Bundle.ToString());
 	}
-	EInstallBundleRequestFlags InstallFlags = EInstallBundleRequestFlags::AsyncMount;
+	EInstallBundleRequestFlags InstallFlags = EInstallBundleRequestFlags::AsyncMount | EInstallBundleRequestFlags::ExplicitUpdateList;
 	if (bIsPredownload)
 	{
 		InstallFlags |= EInstallBundleRequestFlags::SkipMount;
diff --git a/Engine/Source/Runtime/InstallBundleManager/Public/InstallBundleTypes.h b/Engine/Source/Runtime/InstallBundleManager/Public/InstallBundleTypes.h
--- a/Engine/Source/Runtime/InstallBundleManager/Public/InstallBundleTypes.h
+++ b/Engine/Source/Runtime/InstallBundleManager/Public/InstallBundleTypes.h
@@ -251,6 +251,7 @@ enum class EInstallBundleRequestFlags : uint32
 	SkipMount = (1 << 5),
 	AsyncMount = (1 << 6),
 	UpdateBundleReports = (1 << 7), // will update the BundleReport for the involved bundles (see InstallBundleManagerReporting.h and IInstallBundleManager::ReportingDelegate)
+	ExplicitUpdateList = (1 << 8),
 	Defaults = None,
 };
 ENUM_CLASS_FLAGS(EInstallBundleRequestFlags)
