[GFP] Remove unused code in FGFPSharedInstallTracker::Release

#rnx

diff --git a/Engine/Plugins/Runtime/GameFeatures/Source/GameFeatures/Private/GameFeaturePluginStateMachine.cpp b/Engine/Plugins/Runtime/GameFeatures/Source/GameFeatures/Private/GameFeaturePluginStateMachine.cpp
--- a/Engine/Plugins/Runtime/GameFeatures/Source/GameFeatures/Private/GameFeaturePluginStateMachine.cpp
+++ b/Engine/Plugins/Runtime/GameFeatures/Source/GameFeatures/Private/GameFeaturePluginStateMachine.cpp
@@ -413,8 +413,6 @@ namespace UE::GameFeatures
 				FPluginInstallBundleReferencers* PluginRefs = InstallBundleToGFPRefs.Find(Bundle);
 				if (PluginRefs)
 				{
-					FNameBuilder BundleName;
-					Bundle.ToString(BundleName);
 					EGFPInstallLevel* Level = PluginRefs->GFPs.FindByHash(PluginNameHash, PluginName);
 					if (Level && *Level >= InLevel)
 					{
