Add a GameFeaturesSubsystem delegate that allows registering observers before the policies start mounting plugins.
This solves a problem where VNI was relying on receiving callbacks from its GFP observer to defer binding of a VNI C++ module to its corresponding Verse code while the plugin containing the VNI module was pending mount, but sometimes plugins with VNI modules are being mounted before the OnPostEngineInit callback it was using to register the observer.

diff --git a/Engine/Plugins/Runtime/GameFeatures/Source/GameFeatures/Private/GameFeaturesSubsystem.cpp b/Engine/Plugins/Runtime/GameFeatures/Source/GameFeatures/Private/GameFeaturesSubsystem.cpp
--- a/Engine/Plugins/Runtime/GameFeatures/Source/GameFeatures/Private/GameFeaturesSubsystem.cpp
+++ b/Engine/Plugins/Runtime/GameFeatures/Source/GameFeatures/Private/GameFeaturesSubsystem.cpp
@@ -38,6 +38,8 @@
 
 DEFINE_LOG_CATEGORY(LogGameFeatures);
 
+FSimpleMulticastDelegate UGameFeaturesSubsystem::OnGameFeaturePolicyPreInit;
+
 namespace UE::GameFeatures
 {
 	static const FString SubsystemErrorNamespace(TEXT("GameFeaturePlugin.Subsystem."));
@@ -690,6 +692,9 @@ void UGameFeaturesSubsystem::OnAssetManagerCreated()
 		UE_CLOG(bHasProject, LogGameFeatures, Error, TEXT("Asset manager settings do not include a rule for assets of type %s, which is required for game feature plugins to function"), *UGameFeatureData::StaticClass()->GetName());
 	}
 
+	// Give system that want to add an observer before we start to mount plugins a chance to do so.
+	OnGameFeaturePolicyPreInit.Broadcast();
+
 	// Create the game-specific policy
 	UE_LOG(LogGameFeatures, Verbose, TEXT("Initializing game features policy (type %s)"), *GameSpecificPolicies->GetClass()->GetName());
 	bInitializedPolicyManager = true; // Set before calling InitGameFeatureManager() because InitGameFeatureManager may load GFPs
diff --git a/Engine/Plugins/Runtime/GameFeatures/Source/GameFeatures/Public/GameFeaturesSubsystem.h b/Engine/Plugins/Runtime/GameFeatures/Source/GameFeatures/Public/GameFeaturesSubsystem.h
--- a/Engine/Plugins/Runtime/GameFeatures/Source/GameFeatures/Public/GameFeaturesSubsystem.h
+++ b/Engine/Plugins/Runtime/GameFeatures/Source/GameFeatures/Public/GameFeaturesSubsystem.h
@@ -3,6 +3,7 @@
 #pragma once
 
 #include "Containers/Union.h"
+#include "Delegates/Delegate.h"
 #include "Engine/Engine.h"
 #include "GameFeatureTypesFwd.h"
 #include "Misc/TransactionallySafeRWLock.h"
@@ -443,6 +444,8 @@ public:
 
 	static UGameFeaturesSubsystem& Get() { return *GEngine->GetEngineSubsystem<UGameFeaturesSubsystem>(); }
 
+	UE_API static FSimpleMulticastDelegate OnGameFeaturePolicyPreInit;
+
 public:
 	/** Loads the specified game feature data and its bundles */
 	static UE_API TSharedPtr<FStreamableHandle> LoadGameFeatureData(const FString& GameFeatureToLoad, bool bStartStalled = false);
