-Correctly add ContentBundle/AssetRegistry tags to GameFeatureData Objects (previously would miss tags in sub GameFeatureActions)
-Implemented ContentBundle version of the LevelExternalActorsPathsProvider delegate

-Added possibility for implementers of IAssetDependencyGatherers to receive the CachedPathTree from the AssetRegisrty (to help accelerate their calculations) by chaging the param list to a context. New implementers only need to implement this one, the old one is only used for deprecation when the new one is not implemented. Implementers still using the previous signature wil cause a deprecation log message to be outputted (once) when invoked.
-Updated all the UE implementers to use the new signature.
-Added IExternalAssetPathsProvider interface, implementers are tasked with providing the external assets (ExternalPackaged objects) for a specific package
-Changed FExternalObjectAndActorDependencyGatherer to use an IExternalAssetPathsProvider provided by GameFeatureData to provide ExternalDataLayers & ContentBundle paths injected by GameFeatureActions
-Added  FGameFeatureDataExternalAssetsPathCache an implementation of IExternalAssetPathsProvider used by FExternalObjectAndActorDependencyGatherer  to discover dependencies quickly
-Updated on first call to FExternalObjectAndActorDependencyGatherer::GatherDependencies when invalid
-Invalidated when new paths are added
-FGameFeatureDataExternalAssetsPathCache is owed by GameFeaturesSubsystem

diff --git a/Engine/Plugins/Runtime/GameFeatures/Source/GameFeatures/Private/GameFeatureData.cpp b/Engine/Plugins/Runtime/GameFeatures/Source/GameFeatures/Private/GameFeatureData.cpp
--- a/Engine/Plugins/Runtime/GameFeatures/Source/GameFeatures/Private/GameFeatureData.cpp
+++ b/Engine/Plugins/Runtime/GameFeatures/Source/GameFeatures/Private/GameFeatureData.cpp
@@ -26,6 +26,9 @@
 #include "GameFeatureAction_AddWPContent.h"
 #include "Misc/DataValidation.h"
 #include "Engine/Level.h"
+#include "AssetRegistry/AssetRegistryState.h"
+#include "AssetRegistry/PathTree.h"
+#include "Misc/PathViews.h"
 #endif
 
 #include UE_INLINE_GENERATED_CPP_BY_NAME(GameFeatureData)
@@ -716,36 +719,11 @@ TArray<UClass*> UGameFeatureData::GetDisallowedActions() const
 	return DisallowedClasses;
 }
 
-FName UGameFeatureData::GetContentBundleGuidsAssetRegistryTagPrivate()
-{
-	static const FName ContentBundlesTag("ContentBundleGuids");
-	return ContentBundlesTag;
-}
-
-void UGameFeatureData::GetContentBundleGuids(const FAssetData& Asset, TArray<FGuid>& OutContentBundleGuids)
-{
-	FString ContentBundleGuidsStr;
-	if (Asset.GetTagValue(GetContentBundleGuidsAssetRegistryTagPrivate(), ContentBundleGuidsStr))
-	{
-		TArray<FString> ContentBundleGuidsStrArray;
-		ContentBundleGuidsStr.ParseIntoArray(ContentBundleGuidsStrArray, TEXT(","));
-		for (const FString& GuidStr : ContentBundleGuidsStrArray)
-		{
-			FGuid ContentBundleGuid;
-			if (FGuid::Parse(GuidStr, ContentBundleGuid))
-			{
-				OutContentBundleGuids.Add(ContentBundleGuid);
-			}
-		}
-	}
-}
-
 void UGameFeatureData::GetDependencyDirectoriesFromAssetData(const FAssetData& AssetData, TArray<FString>& OutDependencyDirectories)
 {
 	const FString MountPoint = FPackageName::GetPackageMountPoint(AssetData.PackagePath.ToString()).ToString();
 
-	TArray<FGuid> ContentBundleGuids;
-	GetContentBundleGuids(AssetData, ContentBundleGuids);
+	TArray<FGuid> ContentBundleGuids = ContentBundlePaths::ParseContentBundleGuids(AssetData);
 	for (const FGuid& ContentBundleGuid : ContentBundleGuids)
 	{
 		FString ContentBundleExternalActorPath;
@@ -783,9 +761,9 @@ void UGameFeatureData::GetAssetRegistryTags(FAssetRegistryTagsContext Context) c
 	TArray<FGuid> ContentBundleGuids;
 	TArray<FExternalDataLayerUID> ExternalDataLayerUIDs;
 
-	for (UGameFeatureAction* Action : Actions)
+	ForEachObjectWithOuter(this, [&ExternalDataLayerUIDs, &ContentBundleGuids](UObject* Object)
 	{
-		if (UGameFeatureAction_AddWPContent* WPAction = Cast<UGameFeatureAction_AddWPContent>(Action))
+		if (UGameFeatureAction_AddWPContent* WPAction = Cast<UGameFeatureAction_AddWPContent>(Object))
 		{
 			if (const UContentBundleDescriptor* ContentBundleDescriptor = WPAction->GetContentBundleDescriptor())
 			{
@@ -793,21 +771,16 @@ void UGameFeatureData::GetAssetRegistryTags(FAssetRegistryTagsContext Context) c
 			}
 		}
 
-		if (UGameFeatureAction_AddWorldPartitionContent* WPAction = Cast<UGameFeatureAction_AddWorldPartitionContent>(Action))
+		if (UGameFeatureAction_AddWorldPartitionContent* WPAction = Cast<UGameFeatureAction_AddWorldPartitionContent>(Object))
 		{
 			if (const UExternalDataLayerAsset* ExternalDataLayerAsset = WPAction->GetExternalDataLayerAsset())
 			{
 				ExternalDataLayerUIDs.Add(ExternalDataLayerAsset->GetUID());
 			}
 		}
-	}
-
-	if (ContentBundleGuids.Num() > 0)
-	{
-		FString ContentBundleGuidsStr = FString::JoinBy(ContentBundleGuids, TEXT(","), [&](const FGuid& Guid) { return Guid.ToString(); });
-		Context.AddTag(FAssetRegistryTag(GetContentBundleGuidsAssetRegistryTagPrivate(), ContentBundleGuidsStr, FAssetRegistryTag::TT_Hidden));
-	}
+	});
 
+	ContentBundlePaths::AddRegistryTags(Context, ContentBundleGuids);
 	FExternalDataLayerHelper::AddAssetRegistryTags(Context, ExternalDataLayerUIDs);
 }
 #endif
@@ -903,5 +876,176 @@ bool UGameFeatureData::IsGameFeaturePluginActive(bool bCheckForActivating /*= fa
 	return false;
 }
 
+#if WITH_EDITOR
+FGameFeatureDataExternalAssetsPathCache::FGameFeatureDataExternalAssetsPathCache()
+{
+	// Install delegate on PathsAdded event
+	// This is the only thing we watch to invalidate the cache
+	// 
+	// Reasoning:
+	// 
+	// Any newly added Actor in a previously absent EDL or CB for a level causes a directory to be created so we get invalidated and the path map is updated
+	// Removed actors leave empty paths behind temporarily , we'll report a few empty paths, but no wrong dependency will be created from those empty paths
+	// Newly mounted/unmounted plugins also fire this event so that'll update the cache 
+	FExternalObjectAndActorDependencyGatherer::SetExternalAssetPathsProvider(this);
+
+	// Of note....
+	// 
+	// If mounting 2 plugins in sequence, and one of those plugins has dependencies on the other
+	// plugin there's a possibility we'll be gathering and updating the cache before the 2nd plugin is 
+	// mounted. 
+	// 
+	// For example if plugin A has a world with EDL content in plugin B, if the dependencies for world A
+	// are gathered before the plugin B is mounted we'll not detect it's actors and the gather 
+	// will not be redone on world A after plugin B is mounted. 
+}
+
+FGameFeatureDataExternalAssetsPathCache::~FGameFeatureDataExternalAssetsPathCache()
+{
+	if (IAssetRegistry* AssetRegistry = IAssetRegistry::Get())
+	{
+		AssetRegistry->OnPathsAdded().Remove(OnPathAddedDelegateHandle);
+	}
+
+	FExternalObjectAndActorDependencyGatherer::SetExternalAssetPathsProvider(nullptr);
+}
+
+void FGameFeatureDataExternalAssetsPathCache::OnPathsAdded(TConstArrayView<FStringView>)
+{
+	// next update will rebuild it
+	bCacheIsUpToDate = false;
+}
+
+void FGameFeatureDataExternalAssetsPathCache::UpdateCache(const FUpdateCacheContext& Context)
+{	
+	if (bCacheIsUpToDate)
+	{
+		return;
+	}
+
+	TRACE_CPUPROFILER_EVENT_SCOPE(FGameFeatureDataExternalAssetsPathCache::UpdateCache);
+
+	if (!OnPathAddedDelegateHandle.IsValid())
+	{
+		OnPathAddedDelegateHandle = IAssetRegistry::Get()->OnPathsAdded().AddRaw(this, &FGameFeatureDataExternalAssetsPathCache::OnPathsAdded);
+	}
+
+	bCacheIsUpToDate= true;
+	PerLevelAssetDirectories.Reset();
+
+	auto EnumerateAllAssetsOfTypeRecursive = [&Context](UClass* Class, bool bNativeOnly, auto AssetLambda)
+	{
+		if(bNativeOnly)
+		{
+			TArray<UClass*> DerivedClasses;
+			GetDerivedClasses(Class, DerivedClasses, true);
+			DerivedClasses.Add(Class);
+
+			for(UClass* DerivedClass : DerivedClasses)
+			{
+				Context.AssetRegistryState.EnumerateAssetsByClassPathName(DerivedClass->GetClassPathName(), [&AssetLambda](const FAssetData* AssetData)->bool
+				{
+					return AssetLambda(*AssetData);
+				});
+			}
+		}
+		else
+		{
+			FARFilter Filter;	
+			Filter.bRecursiveClasses = true;
+			Filter.ClassPaths = { Class->GetClassPathName() };
+				
+			Context.AssetRegistryState.EnumerateAssets(Context.CompileFilterFunc(Filter), {}, AssetLambda);
+		}
+	};
+
+	// Build set of all levels for validation while discovering them in external actors
+	EnumerateAllAssetsOfTypeRecursive(UWorld::StaticClass(), true, [this](const FAssetData& AssetData)->bool
+	{
+		AllLevels.Add(AssetData.PackageName);
+		return true;
+	});
+
+	EnumerateAllAssetsOfTypeRecursive(UGameFeatureData::StaticClass(), true, [this, &Context](const FAssetData& AssetData)->bool
+	{
+		auto GetLevelSubPaths = [this, &Context](FStringView RootPath, FStringView ParentPath, FName ParentPathName, auto GetLevelSubPathsFn)->void
+		{
+			Context.CachedPathTree.EnumerateSubPaths(ParentPathName, [&RootPath, &ParentPath, &GetLevelSubPathsFn,this, &ParentPathName](FName SubPath)
+			{
+				// Identify folder pattern where we switch from __ExternalActors__/EDL/EDLUID/The/Map/Path/MapName to the sub folders for External Actors
+				// And verify against the Level set we found a real level path (just in case somebody likes really short paths)
+				TStringBuilder<256> SubPathStringBuilder;
+				SubPathStringBuilder << SubPath;
+
+				FStringView LeafName = FPathViews::GetCleanFilename(FStringView(SubPathStringBuilder));
+				if (LeafName.Len() == 1)
+				{
+					FStringView LevelString = ParentPath.RightChop(RootPath.Len());
+					FName LevelPathName (LevelString);
+
+					// could be a level
+					if (AllLevels.Contains(LevelPathName))
+					{
+						// keep level 
+						PerLevelAssetDirectories.Add(LevelPathName, ParentPathName);
+
+						// once we've found one, there no point continuing
+						return false;
+					}
+				}
+				
+				GetLevelSubPathsFn(RootPath, SubPathStringBuilder, SubPath, GetLevelSubPathsFn);
+
+				return true;
+			}, false);
+		};
+
+		TArray<FExternalDataLayerUID> EDLUIDs;
+		FExternalDataLayerHelper::GetExternalDataLayerUIDs(AssetData, EDLUIDs);
+		TArray<FGuid> ContentBundleGuids = ContentBundlePaths::ParseContentBundleGuids(AssetData);
+
+		if (EDLUIDs.Num() || ContentBundleGuids.Num())
+		{
+			FString& MountPoint = GameFeatureDataAssetsToMountPoint.FindOrAdd(AssetData.PackagePath);
+			if (MountPoint.IsEmpty())
+			{
+				MountPoint = FPackageName::GetPackageMountPoint(AssetData.PackagePath.ToString()).ToString();
+			}
+
+			for (const FExternalDataLayerUID& EDLUID : EDLUIDs)
+			{
+				FString EDLRootPath;
+				if(FExternalDataLayerHelper::BuildExternalDataLayerActorsRootPath(MountPoint, EDLUID, EDLRootPath))
+				{
+					FName EDLRootPathName(EDLRootPath);
+					GetLevelSubPaths(EDLRootPath, EDLRootPath, EDLRootPathName, GetLevelSubPaths);
+				}
+			}
+
+		
+			for (const FGuid& ContentBundleGuid : ContentBundleGuids)
+			{
+				FString CBRootPath;
+				if (ContentBundlePaths::BuildContentBundleActorsRootPath(MountPoint, ContentBundleGuid, CBRootPath))
+				{
+					FName CBRootPathName(CBRootPath);
+					GetLevelSubPaths(CBRootPath, CBRootPath, CBRootPathName, GetLevelSubPaths);
+				}
+			}
+		}
+		return true;
+	});
+}
+
+TArray<FName> FGameFeatureDataExternalAssetsPathCache::GetPathsForPackage(FName LevelPath)
+{
+	ensure(bCacheIsUpToDate);
+	TArray<FName> Paths;
+	PerLevelAssetDirectories.MultiFind(LevelPath, Paths);
+	return Paths;
+}
+
+#endif //#if WITH_EDITOR
+
 #undef LOCTEXT_NAMESPACE
 
diff --git a/Engine/Plugins/Runtime/GameFeatures/Source/GameFeatures/Private/GameFeatureDataAssetDependencyGatherer.cpp b/Engine/Plugins/Runtime/GameFeatures/Source/GameFeatures/Private/GameFeatureDataAssetDependencyGatherer.cpp
--- a/Engine/Plugins/Runtime/GameFeatures/Source/GameFeatures/Private/GameFeatureDataAssetDependencyGatherer.cpp
+++ b/Engine/Plugins/Runtime/GameFeatures/Source/GameFeatures/Private/GameFeatureDataAssetDependencyGatherer.cpp
@@ -11,28 +11,28 @@
 // Register FGameFeatureDataAssetDependencyGatherer for UGameFeatureData class
 REGISTER_ASSETDEPENDENCY_GATHERER(FGameFeatureDataAssetDependencyGatherer, UGameFeatureData);
 
-void FGameFeatureDataAssetDependencyGatherer::GatherDependencies(const FAssetData& AssetData, const FAssetRegistryState& AssetRegistryState, TFunctionRef<FARCompiledFilter(const FARFilter&)> CompileFilterFunc, TArray<IAssetDependencyGatherer::FGathereredDependency>& OutDependencies, TArray<FString>& OutDependencyDirectories) const
+void FGameFeatureDataAssetDependencyGatherer::GatherDependencies(FGatherDependenciesContext& Context) const
 {
 	FARFilter Filter;
 	Filter.bRecursivePaths = true;
 	Filter.bIncludeOnlyOnDiskAssets = true;
 
 	TArray<FString> DependencyDirectories;
-	UGameFeatureData::GetDependencyDirectoriesFromAssetData(AssetData, DependencyDirectories);
+	UGameFeatureData::GetDependencyDirectoriesFromAssetData(Context.GetAssetData(), DependencyDirectories);
 	for (const FString& DependencyDirectory : DependencyDirectories)
 	{
-		OutDependencyDirectories.Add(DependencyDirectory);
+		Context.GetOutDependencyDirectories().Add(DependencyDirectory);
 		Filter.PackagePaths.Add(*DependencyDirectory);
 	}
 
 	if (Filter.PackagePaths.Num() > 0)
 	{
 		TArray<FAssetData> FilteredAssets;
-		AssetRegistryState.GetAssets(CompileFilterFunc(Filter), {}, FilteredAssets, true);
+		Context.GetAssetRegistryState().GetAssets(Context.CompileFilter(Filter), {}, FilteredAssets, true);
 
 		for (const FAssetData& FilteredAsset : FilteredAssets)
 		{
-			OutDependencies.Emplace(IAssetDependencyGatherer::FGathereredDependency{ FilteredAsset.PackageName,
+			Context.GetOutDependencies().Emplace(IAssetDependencyGatherer::FGathereredDependency{ FilteredAsset.PackageName,
 				UE::AssetRegistry::EDependencyProperty::Game | UE::AssetRegistry::EDependencyProperty::Build });
 		}
 	}
diff --git a/Engine/Plugins/Runtime/GameFeatures/Source/GameFeatures/Private/GameFeatureDataAssetDependencyGatherer.h b/Engine/Plugins/Runtime/GameFeatures/Source/GameFeatures/Private/GameFeatureDataAssetDependencyGatherer.h
--- a/Engine/Plugins/Runtime/GameFeatures/Source/GameFeatures/Private/GameFeatureDataAssetDependencyGatherer.h
+++ b/Engine/Plugins/Runtime/GameFeatures/Source/GameFeatures/Private/GameFeatureDataAssetDependencyGatherer.h
@@ -13,7 +13,7 @@ public:
 	FGameFeatureDataAssetDependencyGatherer() = default;
 	virtual ~FGameFeatureDataAssetDependencyGatherer() = default;
 
-	virtual void GatherDependencies(const FAssetData& AssetData, const FAssetRegistryState& AssetRegistryState, TFunctionRef<FARCompiledFilter(const FARFilter&)> CompileFilterFunc, TArray<IAssetDependencyGatherer::FGathereredDependency>& OutDependencies, TArray<FString>& OutDependencyDirectories) const override;
+	virtual void GatherDependencies(FGatherDependenciesContext& Context) const override;
 };
 
 #endif
diff --git a/Engine/Plugins/Runtime/GameFeatures/Source/GameFeatures/Private/GameFeaturesSubsystem.cpp b/Engine/Plugins/Runtime/GameFeatures/Source/GameFeatures/Private/GameFeaturesSubsystem.cpp
--- a/Engine/Plugins/Runtime/GameFeatures/Source/GameFeatures/Private/GameFeaturesSubsystem.cpp
+++ b/Engine/Plugins/Runtime/GameFeatures/Source/GameFeatures/Private/GameFeaturesSubsystem.cpp
@@ -657,6 +657,10 @@ void UGameFeaturesSubsystem::Initialize(FSubsystemCollectionBase& Collection)
 	GetExplanationForUnavailablePackageDelegateHandle = IPluginManager::Get().GetExplanationForUnavailablePackageWithPluginInfoDelegate().AddUObject(this, &UGameFeaturesSubsystem::GetExplanationForUnavailablePackage);
 	IPluginManager::Get().OnPluginUnmounted().AddWeakLambda(this, [this] (const IPlugin& Plugin) { PruneCachedGameFeaturePluginDetails(/* doesnt use URL */ FString(""), Plugin.GetDescriptorFileName()); });
 	IPluginManager::Get().OnPluginEdited().AddWeakLambda(this, [this](IPlugin& Plugin) { PruneCachedGameFeaturePluginDetails(/* doesnt use URL */ FString(""), Plugin.GetDescriptorFileName()); });
+
+#if WITH_EDITOR
+	GameFeatureDataExternalAssetsPathCache = MakeUnique<FGameFeatureDataExternalAssetsPathCache>();
+#endif
 }
 
 void UGameFeaturesSubsystem::Deinitialize()
@@ -677,6 +681,10 @@ void UGameFeaturesSubsystem::Deinitialize()
 		FTSTicker::RemoveTicker(MoveTemp(TickHandle));
 		TickHandle.Reset();
 	}
+
+#if WITH_EDITOR
+	GameFeatureDataExternalAssetsPathCache.Reset();
+#endif
 }
 
 void UGameFeaturesSubsystem::OnAssetManagerCreated()
diff --git a/Engine/Plugins/Runtime/GameFeatures/Source/GameFeatures/Public/GameFeatureData.h b/Engine/Plugins/Runtime/GameFeatures/Source/GameFeatures/Public/GameFeatureData.h
--- a/Engine/Plugins/Runtime/GameFeatures/Source/GameFeatures/Public/GameFeatureData.h
+++ b/Engine/Plugins/Runtime/GameFeatures/Source/GameFeatures/Public/GameFeatureData.h
@@ -4,6 +4,7 @@
 
 #include "Engine/DataAsset.h"
 #include "GameFeatureAction.h"
+#include "Engine/ExternalObjectAndActorDependencyGatherer.h" // For IExternalAssetPathsProvider
 
 #include "GameFeatureData.generated.h"
 
@@ -15,6 +16,28 @@ struct FExternalDataLayerUID;
 
 struct FAssetData;
 
+#if WITH_EDITOR
+
+class FPathTree;
+class FGameFeatureDataExternalAssetsPathCache : public IExternalAssetPathsProvider
+{
+	TMultiMap<FName, FName> PerLevelAssetDirectories;
+	TMap<FName, FString> GameFeatureDataAssetsToMountPoint;
+	TSet<FName> AllLevels;
+	bool bCacheIsUpToDate = false;
+	FDelegateHandle OnPathAddedDelegateHandle;
+
+public:
+		
+	FGameFeatureDataExternalAssetsPathCache();
+	virtual ~FGameFeatureDataExternalAssetsPathCache();
+
+	void OnPathsAdded(TConstArrayView<FStringView>);
+	void UpdateCache(const FUpdateCacheContext& Context) override;
+	TArray<FName> GetPathsForPackage(FName LevelPath) override;
+};
+#endif
+
 /** Data related to a game feature, a collection of code and content that adds a separable discrete feature to the game */
 UCLASS(MinimalAPI)
 class UGameFeatureData : public UPrimaryDataAsset
@@ -93,9 +116,6 @@ public:
 
 private:
 #if WITH_EDITOR
-	static UE_API void GetContentBundleGuids(const FAssetData& Asset, TArray<FGuid>& OutContentBundleGuids);
-	static UE_API FName GetContentBundleGuidsAssetRegistryTagPrivate();
-
 	UFUNCTION()
 	UE_API TArray<UClass*> GetDisallowedActions() const;
 #endif
diff --git a/Engine/Plugins/Runtime/GameFeatures/Source/GameFeatures/Public/GameFeaturesSubsystem.h b/Engine/Plugins/Runtime/GameFeatures/Source/GameFeatures/Public/GameFeaturesSubsystem.h
--- a/Engine/Plugins/Runtime/GameFeatures/Source/GameFeatures/Public/GameFeaturesSubsystem.h
+++ b/Engine/Plugins/Runtime/GameFeatures/Source/GameFeatures/Public/GameFeaturesSubsystem.h
@@ -954,6 +954,8 @@ private:
 #if WITH_EDITOR
 	// When we decide not to mount a plugin, we can store an explanation here so that if we later attempt to load an asset from it we can tell the user why it's not available
 	TMap<FString, FString> UnmountedPluginNameToExplanation;
+
+	TUniquePtr<class FGameFeatureDataExternalAssetsPathCache> GameFeatureDataExternalAssetsPathCache;
 #endif
 
 #if !UE_BUILD_SHIPPING
diff --git a/Engine/Source/Editor/WorldPartitionEditor/Private/WorldPartition/ContentBundle/ContentBundleEditorSubsystem.cpp b/Engine/Source/Editor/WorldPartitionEditor/Private/WorldPartition/ContentBundle/ContentBundleEditorSubsystem.cpp
--- a/Engine/Source/Editor/WorldPartitionEditor/Private/WorldPartition/ContentBundle/ContentBundleEditorSubsystem.cpp
+++ b/Engine/Source/Editor/WorldPartitionEditor/Private/WorldPartition/ContentBundle/ContentBundleEditorSubsystem.cpp
@@ -2,6 +2,7 @@
 #include "WorldPartition/ContentBundle/ContentBundleEditorSubsystem.h"
 #include "WorldPartition/ContentBundle/ContentBundleEngineSubsystem.h"
 #include "WorldPartition/ContentBundle/ContentBundleWorldSubsystem.h"
+#include "WorldPartition/ContentBundle/ContentBundlePaths.h"
 #include "WorldPartition/ContentBundle/ContentBundleDescriptor.h"
 #include "WorldPartition/ContentBundle/ContentBundleEditor.h"
 #include "WorldPartition/ContentBundle/ContentBundleLog.h"
@@ -364,6 +365,8 @@ void UContentBundleEditorSubsystem::Initialize(FSubsystemCollectionBase& Collect
 	ContentBundleEditingSubModule = NewObject<UContentBundleEditingSubmodule>(this, NAME_None, RF_Transactional);
 	ContentBundleEditingSubModule->Initialize();
 
+	LevelExternalActorsPathsProviderDelegateHandle = ULevel::RegisterLevelExternalActorsPathsProvider(ULevel::FLevelExternalActorsPathsProviderDelegate::CreateUObject(this, &UContentBundleEditorSubsystem::OnGetLevelExternalActorsPaths));
+
 	IContentBundleEditorSubsystemInterface::SetInstance(this);
 }
 
@@ -371,6 +374,8 @@ void UContentBundleEditorSubsystem::Deinitialize()
 {
 	IContentBundleEditorSubsystemInterface::SetInstance(nullptr);
 
+	ULevel::UnregisterLevelExternalActorsPathsProvider(LevelExternalActorsPathsProviderDelegateHandle);
+
 	ContentBundleEditingSubModule->Deinitialize();
 	ContentBundleEditingSubModule = nullptr;
 }
@@ -521,4 +526,14 @@ void UContentBundleEditorSubsystem::SelectActorsInternal(FContentBundleEditor& E
 
 	GEditor->NoteSelectionChange();
 }
+
+void UContentBundleEditorSubsystem::OnGetLevelExternalActorsPaths(const FString& InLevelPackageName, const FString& InPackageShortName, TArray<FString>& OutExternalActorsPaths)
+{
+	ContentBundlePaths::ForEachContentBundleLevelPackagePath(InLevelPackageName, [&OutExternalActorsPaths](const FString& Path)
+	{
+		const FString ExternalActorPath = ULevel::GetExternalActorsPath(Path);
+		OutExternalActorsPaths.Add(ExternalActorPath);
+	});
+}
+
 #undef LOCTEXT_NAMESPACE
diff --git a/Engine/Source/Editor/WorldPartitionEditor/Public/WorldPartition/ContentBundle/ContentBundleEditorSubsystem.h b/Engine/Source/Editor/WorldPartitionEditor/Public/WorldPartition/ContentBundle/ContentBundleEditorSubsystem.h
--- a/Engine/Source/Editor/WorldPartitionEditor/Public/WorldPartition/ContentBundle/ContentBundleEditorSubsystem.h
+++ b/Engine/Source/Editor/WorldPartitionEditor/Public/WorldPartition/ContentBundle/ContentBundleEditorSubsystem.h
@@ -154,6 +154,10 @@ public:
 	UE_API virtual void PopContentBundleEditing() override;
 
 private:
+
+	FDelegateHandle LevelExternalActorsPathsProviderDelegateHandle;
+	void OnGetLevelExternalActorsPaths(const FString& InLevelPackageName, const FString& InPackageShortName, TArray<FString>& OutExternalActorsPaths);
+
 	UE_API void SelectActorsInternal(FContentBundleEditor& EditorContentBundle, bool bSelect);
 
 	UPROPERTY()
diff --git a/Engine/Source/Runtime/AssetRegistry/Private/AssetDependencyGatherer.cpp b/Engine/Source/Runtime/AssetRegistry/Private/AssetDependencyGatherer.cpp
--- a/Engine/Source/Runtime/AssetRegistry/Private/AssetDependencyGatherer.cpp
+++ b/Engine/Source/Runtime/AssetRegistry/Private/AssetDependencyGatherer.cpp
@@ -1,6 +1,7 @@
 // Copyright Epic Games, Inc. All Rights Reserved.
 
 #include "AssetRegistry/AssetDependencyGatherer.h"
+#include "AssetRegistryPrivate.h"
 
 #if WITH_EDITOR
 
@@ -37,4 +38,52 @@ void FRegisteredAssetDependencyGatherer::ForEach(TFunctionRef<void(FRegisteredAs
 	
 }
 
+namespace UE::AssetDependencyGatherer::Private
+{
+
+void FRegisteredAssetDependencyGatherer::GatherDependenciesForGatherer(IAssetDependencyGatherer* GathererInterface, IAssetDependencyGatherer::FGatherDependenciesContext& Context, const char* InGathererClass) const 
+{ 
+	// reset marker that indicates if the default implementation in IAssetDependencyGatherer was called
+	Context.bInterfaceGatherDependenciesCalled = false;
+	GathererInterface->GatherDependencies(Context); 
+
+	if (Context.bInterfaceGatherDependenciesCalled)
+	{
+		PRAGMA_DISABLE_DEPRECATION_WARNINGS
+
+		// We call the deprecated implementation if we detect a call that directly ended in IAssetDependencyGatherer::GatherDependencies(FGatherDependenciesContext&)
+		// This shows the IAssetDependencyGatherer we have is still using the old interface. 
+		GathererInterface->GatherDependencies(Context.AssetData, Context.AssetRegistryState, Context.CompileFilterFunc, Context.OutDependencies, Context.OutDependencyDirectories);
+
+		PRAGMA_ENABLE_DEPRECATION_WARNINGS
+
+		if (!bCalledOnce)
+		{
+			// Deprecated in 5.8
+			UE_LOG(LogAssetRegistry, Warning, 
+				TEXT("%hs is overriding a deprecated function, GatherDependencies(const FAssetData& , const FAssetRegistryState&, TArray&, TArray&). ")
+				TEXT("Override GatherDependencies that takes an FGatherDependenciesContext instead. ")
+				TEXT("Deprecation handling will be removed in a future version and your code will no longer compile."),
+				InGathererClass);
+
+			bCalledOnce = true;
+		}
+	}
+}
+
+}
+
+
+void IAssetDependencyGatherer::GatherDependencies(FGatherDependenciesContext& Context) const
+{
+	// Flag the context to indicate IAssetDependencyGatherer::GatherDependencies(FGatherDependenciesContext&) was called which will 
+	// happen in classes only implementing GatherDependencies(const FAssetData&, const FAssetRegistryState&, TFunctionRef<FARCompiledFilter(const FARFilter&)>, TArray<FGathereredDependency>&, TArray<FString>&)
+	Context.bInterfaceGatherDependenciesCalled = true;
+}
+
+void IAssetDependencyGatherer::GatherDependencies(const FAssetData & AssetData, const FAssetRegistryState & AssetRegistryState, TFunctionRef<FARCompiledFilter(const FARFilter&)> CompileFilterFunc, TArray<FGathereredDependency>&OutDependencies, TArray<FString>&OutDependencyDirectories) const
+{
+	checkf(false, TEXT("Implement IAssetDependencyGatherer::GatherDependencies(FGatherDependenciesContext& Context) in your class, this class currently appears to not be implementing any."));
+}
+
 #endif
\ No newline at end of file
diff --git a/Engine/Source/Runtime/AssetRegistry/Private/AssetRegistry.cpp b/Engine/Source/Runtime/AssetRegistry/Private/AssetRegistry.cpp
--- a/Engine/Source/Runtime/AssetRegistry/Private/AssetRegistry.cpp
+++ b/Engine/Source/Runtime/AssetRegistry/Private/AssetRegistry.cpp
@@ -6292,13 +6292,18 @@ void FAssetRegistryImpl::LoadCalculatedDependencies(FName PackageName,
 		// This code could be moved somewhere where it doesn't need to re-query the asset data, but it needs to happen after both dependencies and data are handled
 		RegisteredDependencyGathererClasses.MultiFind(AssetData->AssetClassPath, Gatherers);
 
+		TArray<IAssetDependencyGatherer::FGathereredDependency> GatheredDependencies;
+		TArray<FString> GathererDirectoryReferences;
+
+		IAssetDependencyGatherer::FGatherDependenciesContext Context { *AssetData, State, GetCachedPathTree(), GetCompiledFilter, GatheredDependencies, GathererDirectoryReferences};
+
 		for (UE::AssetDependencyGatherer::Private::FRegisteredAssetDependencyGatherer* Gatherer : Gatherers)
 		{
 			bOutHadActivity = true;
-			TArray<IAssetDependencyGatherer::FGathereredDependency> GatheredDependencies;
+			GatheredDependencies.Reset();
+			GathererDirectoryReferences.Reset();
 
-			TArray<FString> GathererDirectoryReferences;
-			Gatherer->GatherDependencies(*AssetData, State, GetCompiledFilter, GatheredDependencies, GathererDirectoryReferences);
+			Gatherer->GatherDependencies(Context);
 			
 			if (GatheredDependencies.Num())
 			{
diff --git a/Engine/Source/Runtime/AssetRegistry/Public/AssetRegistry/AssetDependencyGatherer.h b/Engine/Source/Runtime/AssetRegistry/Public/AssetRegistry/AssetDependencyGatherer.h
--- a/Engine/Source/Runtime/AssetRegistry/Public/AssetRegistry/AssetDependencyGatherer.h
+++ b/Engine/Source/Runtime/AssetRegistry/Public/AssetRegistry/AssetDependencyGatherer.h
@@ -20,8 +20,18 @@ class UClass;
 struct FARCompiledFilter;
 struct FARFilter;
 struct FAssetData;
+class FPathTree; 
 template <typename FuncType> class TFunctionRef;
 
+namespace UE::AssetDependencyGatherer::Private
+{
+	class FRegisteredAssetDependencyGatherer;
+};
+
+namespace UE::AssetRegistry
+{
+	class FAssetRegistryImpl;
+}
 /**
  * Interface class for functions that return extra dependencies to add to Assets of a given Class, 
  * to e.g. add dependencies on external Assets in a subdirectory. 
@@ -34,6 +44,63 @@ public:
 		FName PackageName;
 		UE::AssetRegistry::EDependencyProperty Property;
 	};
+		
+	/* The parameters passed to GatherDependencies
+	 *
+	 * @param AssetData The Asset to which the dependencies will be added
+	 * @param AssetRegistryState Helper object: The AssetRegistry's internal state which can be used to run queries
+	 *        in the callback. This must be used instead of the usual interface functions on
+	 *        IAssetRegistry::GetChecked(), because the callback executes inside the critical section.
+	 * @param CompileFilterFunc Helper object: Used to compile FARFilter into FARCompiledFilter for queries sent to
+	 *        AssetRegistryState. The same functionality that is normally provided by
+	 *        IAssetRegistry::GetChecked().CompileFilter.
+	 * @param CachedPathTree The path tree cache inside the AssetRegitry can be used to quickly interogate the path 
+	 *        hierarchy. 
+	 * @param OutDependencies List of gathered dependencies to add for the asset.
+	 * @param OutDependencyDirectories List of directories that the callback queried. The AssetRegistry will call the callback
+	 *        again if any assets are added to any of these directories.	 
+	 */
+	struct FGatherDependenciesContext
+	{
+		friend class UE::AssetRegistry::FAssetRegistryImpl;
+		friend class IAssetDependencyGatherer;
+		friend class UE::AssetDependencyGatherer::Private::FRegisteredAssetDependencyGatherer;
+		
+	private:
+		
+		const FAssetData& AssetData;
+		const FAssetRegistryState& AssetRegistryState;
+		const FPathTree& CachedPathTree;
+		const TFunctionRef<FARCompiledFilter(const FARFilter&)> CompileFilterFunc;
+		TArray<FGathereredDependency>& OutDependencies;
+		TArray<FString>& OutDependencyDirectories;
+
+		// helper flag to detect old implementations of IAssetDependencyGatherer
+		bool bInterfaceGatherDependenciesCalled = false;
+		
+	public:
+
+		FGatherDependenciesContext(	const FAssetData& InAssetData, const FAssetRegistryState& InAssetRegistryState, 
+									const FPathTree& InCachedPathTree, const TFunctionRef<FARCompiledFilter(const FARFilter&)> InCompileFilterFunc, 
+									TArray<FGathereredDependency>& InReturnedDependencies, TArray<FString>& InReturnedDependencyDirectories) :
+									AssetData(InAssetData),
+									AssetRegistryState(InAssetRegistryState),
+									CachedPathTree(InCachedPathTree),
+									CompileFilterFunc(InCompileFilterFunc),
+									OutDependencies(InReturnedDependencies),
+									OutDependencyDirectories(InReturnedDependencyDirectories)
+		{
+		}
+
+		const FAssetData& GetAssetData() const { return AssetData; }
+		const FAssetRegistryState& GetAssetRegistryState() const { return AssetRegistryState; }
+		const class FPathTree& GetCachedPathTree() const { return CachedPathTree; }
+		TFunctionRef<FARCompiledFilter(const FARFilter&)> GetCompileFilterFunc() const { return CompileFilterFunc;}
+		FARCompiledFilter CompileFilter(FARFilter& Filter) { return CompileFilterFunc(Filter); }
+
+		TArray<FGathereredDependency>& GetOutDependencies() { return OutDependencies; }
+		TArray<FString>& GetOutDependencyDirectories() { return OutDependencyDirectories; }
+	};
 
 	/**
 	 * Return the extra dependencies to add to the given AssetData, which is guaranteed to 
@@ -43,21 +110,21 @@ public:
 	 * Attempting to call public functions on the AssetRegistry will deadlock. 
 	 * To send queries about what assets exist, used the passed-in interface functions instead.
 	 * 
-	 * @param AssetData The Asset to which the dependencies will be added
-	 * @param AssetRegistryState Helper object: The AssetRegistry's internal state which can be used to run queries
-	 *        in the callback. This must be used instead of the usual interface functions on
-	 *        IAssetRegistry::GetChecked(), because the callback executes inside the critical section.
-	 * @param CompileFilterFunc Helper object: Used to compile FARFilter into FARCompiledFilter for queries sent to
-	 *        AssetRegistryState. The same functionality that is normally provided by
-	 *        IAssetRegistry::GetChecked().CompileFilter.
-	 * @param OutDependencies List of directories that the callback queried. The AssetRegistry will call the callback
-	 *        again if any assets are added to any of these directories.
+	 * @param Context the various inputs/outputs provided to/from the gatherer
+	 
 	 */
-	virtual void GatherDependencies(const FAssetData& AssetData, const FAssetRegistryState& AssetRegistryState, 
+	 ASSETREGISTRY_API virtual void GatherDependencies(FGatherDependenciesContext& Context) const /*= 0*/;
+
+	 /**
+	  * Deprecated version of GatherDependencies, provided for code compatibility. Will be removed
+	  * in UE 5.8. 
+	  */
+	 UE_DEPRECATED(5.8, "Implement GatherDependencies(FGatherDependenciesContext& Context) instead")
+	 ASSETREGISTRY_API virtual void GatherDependencies(const FAssetData& AssetData, const FAssetRegistryState& AssetRegistryState, 
 		TFunctionRef<FARCompiledFilter(const FARFilter&)> CompileFilterFunc, TArray<IAssetDependencyGatherer::FGathereredDependency>& OutDependencies, 
-		TArray<FString>& OutDependencyDirectories) const = 0;
-};
+		TArray<FString>& OutDependencyDirectories) const /*= 0*/;
 
+};
 
 namespace UE::AssetDependencyGatherer::Private
 {
@@ -71,16 +138,17 @@ namespace UE::AssetDependencyGatherer::Private
 		ASSETREGISTRY_API virtual ~FRegisteredAssetDependencyGatherer();
 
 		virtual UClass* GetAssetClass() const = 0;
-		virtual void GatherDependencies(const FAssetData& AssetData, const FAssetRegistryState& AssetRegistryState, 
-			TFunctionRef<FARCompiledFilter(const FARFilter&)> CompileFilterFunc, TArray<IAssetDependencyGatherer::FGathereredDependency>& OutDependencies, 
-			TArray<FString>& OutDependencyDirectories) const = 0;
+		virtual void GatherDependencies(IAssetDependencyGatherer::FGatherDependenciesContext& Context) const = 0;
 
 		static ASSETREGISTRY_API void ForEach(TFunctionRef<void(FRegisteredAssetDependencyGatherer*)> Func);
 		static bool IsEmpty() { return !GetRegisteredList(); }
+
+		ASSETREGISTRY_API void GatherDependenciesForGatherer(IAssetDependencyGatherer* GathererInterface, IAssetDependencyGatherer::FGatherDependenciesContext& Context, const char* InGathererClass) const;
 	private:
 
 		static ASSETREGISTRY_API TLinkedList<FRegisteredAssetDependencyGatherer*>*& GetRegisteredList();
 		TLinkedList<FRegisteredAssetDependencyGatherer*> GlobalListLink;
+		mutable bool bCalledOnce = false;
 	};
 }
 
@@ -97,10 +165,10 @@ namespace UE::AssetDependencyGatherer::Private
 class PREPROCESSOR_JOIN(PREPROCESSOR_JOIN(GathererClass, AssetClass), _Register) : public UE::AssetDependencyGatherer::Private::FRegisteredAssetDependencyGatherer \
 { \
 	virtual UClass* GetAssetClass() const override { return AssetClass::StaticClass(); } \
-	virtual void GatherDependencies(const FAssetData& AssetData, const FAssetRegistryState& AssetRegistryState, TFunctionRef<FARCompiledFilter(const FARFilter&)> CompileFilterFunc, TArray<IAssetDependencyGatherer::FGathereredDependency>& OutDependencies, TArray<FString>& OutDependencyDirectories) const override \
+	virtual void GatherDependencies(IAssetDependencyGatherer::FGatherDependenciesContext& Context) const override \
 	{ \
 		GathererClass Gatherer; \
-		Gatherer.GatherDependencies(AssetData, AssetRegistryState, CompileFilterFunc, OutDependencies, OutDependencyDirectories); \
+		GatherDependenciesForGatherer(&Gatherer, Context, #GathererClass); \
 	} \
 }; \
 namespace PREPROCESSOR_JOIN(GathererClass, AssetClass) \
@@ -108,4 +176,4 @@ namespace PREPROCESSOR_JOIN(GathererClass, AssetClass) \
 	static PREPROCESSOR_JOIN(PREPROCESSOR_JOIN(GathererClass, AssetClass), _Register) DefaultObject; \
 }
 
-#endif // WITH_EDITOR
+#endif // WITH_EDITOR 
diff --git a/Engine/Source/Runtime/Engine/Classes/Engine/ExternalAssetDependencyGatherer.h b/Engine/Source/Runtime/Engine/Classes/Engine/ExternalAssetDependencyGatherer.h
--- a/Engine/Source/Runtime/Engine/Classes/Engine/ExternalAssetDependencyGatherer.h
+++ b/Engine/Source/Runtime/Engine/Classes/Engine/ExternalAssetDependencyGatherer.h
@@ -14,7 +14,7 @@ public:
 	FExternalAssetDependencyGatherer() = default;
 	virtual ~FExternalAssetDependencyGatherer() = default;
 
-	ENGINE_API virtual void GatherDependencies(const FAssetData& AssetData, const FAssetRegistryState& AssetRegistryState, TFunctionRef<FARCompiledFilter(const FARFilter&)> CompileFilterFunc, TArray<IAssetDependencyGatherer::FGathereredDependency>& OutDependencies, TArray<FString>& OutDependencyDirectories) const override;
+	ENGINE_API virtual void GatherDependencies(FGatherDependenciesContext& Context) const override;
 };
 
 #endif
diff --git a/Engine/Source/Runtime/Engine/Classes/Engine/ExternalObjectAndActorDependencyGatherer.h b/Engine/Source/Runtime/Engine/Classes/Engine/ExternalObjectAndActorDependencyGatherer.h
--- a/Engine/Source/Runtime/Engine/Classes/Engine/ExternalObjectAndActorDependencyGatherer.h
+++ b/Engine/Source/Runtime/Engine/Classes/Engine/ExternalObjectAndActorDependencyGatherer.h
@@ -11,19 +11,57 @@
 #include "Engine/World.h"
 
 struct FARFilter;
+struct FARCompiledFilter;
+
+//
+// This interface is used by systems who inject ContentBundle/ExternalDataLayers in worlds to inform
+// FExternalObjectAndActorDependencyGatherer of where to locate those external actors and external objects
+// and which CB/EDLs are associated with which world. 
+//
+// This is necessary since worlds do not store information linking back to which EDL/CB are injected 
+// in them.
+//
+class IExternalAssetPathsProvider
+{
+	public:
+
+		// These functions will be invoked while the AssetRegistry lock is held. All the restrictions 
+		// and warnings pertaining to IAssetDependencyGathers apply here. 
+		// 
+		// Replicated here for clarity....
+		// 
+		// WARNING: For high performance these callbacks are called inside the critical section of the AssetRegistry. 
+		// Attempting to call public functions on the AssetRegistry will deadlock. 
+		// To send queries about what assets exist, used the passed-in interface functions instead.
+		// 
+		// Aside from that these should be made as simple as possible to execute quickly and without side-effects outside
+		// of the cache (if using a cache).
+		//
+
+		struct FUpdateCacheContext
+		{
+			const FAssetRegistryState& AssetRegistryState;
+			const class FPathTree& CachedPathTree;
+			const TFunctionRef<FARCompiledFilter(const FARFilter&)> CompileFilterFunc;
+		};
+
+		virtual void UpdateCache(const FUpdateCacheContext& Context) = 0;
+		virtual TArray<FName> GetPathsForPackage(FName PackagePath) = 0;
+};
 
 class FExternalObjectAndActorDependencyGatherer : public IAssetDependencyGatherer
 {
+	static IExternalAssetPathsProvider* ExternalPathsProvider;
+
 public:	
 	FExternalObjectAndActorDependencyGatherer() = default;
 	virtual ~FExternalObjectAndActorDependencyGatherer() = default;
+	
+	static ENGINE_API void SetExternalAssetPathsProvider(IExternalAssetPathsProvider* InProvider);
 
 	static ENGINE_API FARFilter GetQueryFilter(FName PackageName, TArray<FString>* OutQueryDirectories = nullptr);
 
-	ENGINE_API virtual void GatherDependencies(const FAssetData& AssetData, const FAssetRegistryState& AssetRegistryState,
-		TFunctionRef<FARCompiledFilter(const FARFilter&)> CompileFilterFunc,
-		TArray<IAssetDependencyGatherer::FGathereredDependency>& OutDependencies,
-		TArray<FString>& OutDependencyDirectories) const override;
+	ENGINE_API virtual void GatherDependencies(FGatherDependenciesContext& Context) const override;
 };
 
 #endif
diff --git a/Engine/Source/Runtime/Engine/Private/ExternalActorAssetDependencyGatherer.cpp b/Engine/Source/Runtime/Engine/Private/ExternalActorAssetDependencyGatherer.cpp
--- a/Engine/Source/Runtime/Engine/Private/ExternalActorAssetDependencyGatherer.cpp
+++ b/Engine/Source/Runtime/Engine/Private/ExternalActorAssetDependencyGatherer.cpp
@@ -8,12 +8,12 @@
 #include "AssetRegistry/AssetRegistryState.h"
 #include "Engine/Level.h"
 
-void FExternalActorAssetDependencyGatherer::GatherDependencies(const FAssetData& AssetData, const FAssetRegistryState& AssetRegistryState, TFunctionRef<FARCompiledFilter(const FARFilter&)> CompileFilterFunc, TArray<IAssetDependencyGatherer::FGathereredDependency>& OutDependencies, TArray<FString>& OutDependencyDirectories) const
+void FExternalActorAssetDependencyGatherer::GatherDependencies(FGatherDependenciesContext& Context) const
 {
-	if (ULevel::GetIsLevelUsingExternalActorsFromAsset(AssetData))
+	if (ULevel::GetIsLevelUsingExternalActorsFromAsset(Context.GetAssetData()))
 	{
-		const FString ExternalActorsPath = ULevel::GetExternalActorsPath(AssetData.PackageName.ToString());
-		OutDependencyDirectories.Add(ExternalActorsPath);
+		const FString ExternalActorsPath = ULevel::GetExternalActorsPath(Context.GetAssetData().PackageName.ToString());
+		Context.GetOutDependencyDirectories().Add(ExternalActorsPath);
 
 		FARFilter Filter;
 		Filter.bRecursivePaths = true;
@@ -21,11 +21,11 @@ void FExternalActorAssetDependencyGatherer::GatherDependencies(const FAssetData&
 		Filter.PackagePaths.Add(*ExternalActorsPath);
 
 		TArray<FAssetData> FilteredAssets;
-		AssetRegistryState.GetAssets(CompileFilterFunc(Filter), {}, FilteredAssets, true);
+		Context.GetAssetRegistryState().GetAssets(Context.CompileFilter(Filter), {}, FilteredAssets, true);
 
 		for (const FAssetData& FilteredAsset : FilteredAssets)
 		{
-			OutDependencies.Emplace(IAssetDependencyGatherer::FGathereredDependency{ FilteredAsset.PackageName,
+			Context.GetOutDependencies().Emplace(IAssetDependencyGatherer::FGathereredDependency{ FilteredAsset.PackageName,
 				UE::AssetRegistry::EDependencyProperty::Game  | UE::AssetRegistry::EDependencyProperty::Build });
 		}
 	}
diff --git a/Engine/Source/Runtime/Engine/Private/ExternalActorAssetDependencyGatherer.h b/Engine/Source/Runtime/Engine/Private/ExternalActorAssetDependencyGatherer.h
--- a/Engine/Source/Runtime/Engine/Private/ExternalActorAssetDependencyGatherer.h
+++ b/Engine/Source/Runtime/Engine/Private/ExternalActorAssetDependencyGatherer.h
@@ -15,7 +15,7 @@ public:
 	FExternalActorAssetDependencyGatherer() = default;
 	virtual ~FExternalActorAssetDependencyGatherer() = default;
 
-	virtual void GatherDependencies(const FAssetData& AssetData, const FAssetRegistryState& AssetRegistryState, TFunctionRef<FARCompiledFilter(const FARFilter&)> CompileFilterFunc, TArray<IAssetDependencyGatherer::FGathereredDependency>& OutDependencies, TArray<FString>& OutDependencyDirectories) const override;
+	virtual void GatherDependencies(FGatherDependenciesContext& Context) const override;
 };
 
 #endif
\ No newline at end of file
diff --git a/Engine/Source/Runtime/Engine/Private/ExternalAssetDependencyGatherer.cpp b/Engine/Source/Runtime/Engine/Private/ExternalAssetDependencyGatherer.cpp
--- a/Engine/Source/Runtime/Engine/Private/ExternalAssetDependencyGatherer.cpp
+++ b/Engine/Source/Runtime/Engine/Private/ExternalAssetDependencyGatherer.cpp
@@ -1,4 +1,4 @@
-﻿// Copyright Epic Games, Inc. All Rights Reserved.
+// Copyright Epic Games, Inc. All Rights Reserved.
 #include "Engine/ExternalAssetDependencyGatherer.h"
 
 #if WITH_EDITOR
@@ -7,10 +7,10 @@
 #include "AssetRegistry/AssetRegistryState.h"
 #include "ExternalPackageHelper.h"
 
-void FExternalAssetDependencyGatherer::GatherDependencies(const FAssetData& AssetData, const FAssetRegistryState& AssetRegistryState, TFunctionRef<FARCompiledFilter(const FARFilter&)> CompileFilterFunc, TArray<IAssetDependencyGatherer::FGathereredDependency>& OutDependencies, TArray<FString>& OutDependencyDirectories) const
+void FExternalAssetDependencyGatherer::GatherDependencies(FGatherDependenciesContext& Context) const
 {		
-	const FString ExternalObjectsPath = FExternalPackageHelper::GetExternalObjectsPath(AssetData.PackageName.ToString());
-	OutDependencyDirectories.Add(ExternalObjectsPath);
+	const FString ExternalObjectsPath = FExternalPackageHelper::GetExternalObjectsPath(Context.GetAssetData().PackageName.ToString());
+	Context.GetOutDependencyDirectories().Add(ExternalObjectsPath);
 
 	FARFilter Filter;
 	Filter.bRecursivePaths = true;
@@ -18,11 +18,11 @@ void FExternalAssetDependencyGatherer::GatherDependencies(const FAssetData& Asse
 	Filter.PackagePaths.Add(*ExternalObjectsPath);
 
 	TArray<FAssetData> FilteredAssets;
-	AssetRegistryState.GetAssets(CompileFilterFunc(Filter), {}, FilteredAssets, true);
+	Context.GetAssetRegistryState().GetAssets(Context.CompileFilter(Filter), {}, FilteredAssets, true);
 
 	for (const FAssetData& FilteredAsset : FilteredAssets)
 	{
-		OutDependencies.Emplace(IAssetDependencyGatherer::FGathereredDependency{ FilteredAsset.PackageName, UE::AssetRegistry::EDependencyProperty::Game | UE::AssetRegistry::EDependencyProperty::Build });
+		Context.GetOutDependencies().Emplace(IAssetDependencyGatherer::FGathereredDependency{ FilteredAsset.PackageName, UE::AssetRegistry::EDependencyProperty::Game | UE::AssetRegistry::EDependencyProperty::Build });
 	}
 }
 
diff --git a/Engine/Source/Runtime/Engine/Private/ExternalObjectAndActorDependencyGatherer.cpp b/Engine/Source/Runtime/Engine/Private/ExternalObjectAndActorDependencyGatherer.cpp
--- a/Engine/Source/Runtime/Engine/Private/ExternalObjectAndActorDependencyGatherer.cpp
+++ b/Engine/Source/Runtime/Engine/Private/ExternalObjectAndActorDependencyGatherer.cpp
@@ -7,24 +7,32 @@
 #include "AssetRegistry/AssetRegistryState.h"
 #include "ExternalPackageHelper.h"
 
-void FExternalObjectAndActorDependencyGatherer::GatherDependencies(const FAssetData& AssetData,
-	const FAssetRegistryState& AssetRegistryState,TFunctionRef<FARCompiledFilter(const FARFilter&)> CompileFilterFunc,
-	TArray<IAssetDependencyGatherer::FGathereredDependency>& OutDependencies,
-	TArray<FString>& OutDependencyDirectories) const
+void FExternalObjectAndActorDependencyGatherer::GatherDependencies(FGatherDependenciesContext& Params) const
 {
-	FARFilter Filter = GetQueryFilter(AssetData.PackageName, &OutDependencyDirectories);
+	TRACE_CPUPROFILER_EVENT_SCOPE(FExternalObjectAndActorDependencyGatherer::GatherDependencies);
 
+	if (ExternalPathsProvider)
+	{
+		IExternalAssetPathsProvider::FUpdateCacheContext UpdateContext{Params.GetAssetRegistryState(), Params.GetCachedPathTree(), Params.GetCompileFilterFunc() };
+		ExternalPathsProvider->UpdateCache(UpdateContext);
+	}
+
+	FARFilter Filter = GetQueryFilter(Params.GetAssetData().PackageName, &Params.GetOutDependencyDirectories());
+		
 	TArray<FAssetData> FilteredAssets;
-	AssetRegistryState.GetAssets(CompileFilterFunc(Filter), {}, FilteredAssets, true);
+	Params.GetAssetRegistryState().GetAssets(Params.CompileFilter(Filter), {}, FilteredAssets, true);
 
+	Params.GetOutDependencies().Reserve(FilteredAssets.Num());
 	for (const FAssetData& FilteredAsset : FilteredAssets)
 	{
-		OutDependencies.Emplace(IAssetDependencyGatherer::FGathereredDependency{ FilteredAsset.PackageName, UE::AssetRegistry::EDependencyProperty::Game | UE::AssetRegistry::EDependencyProperty::Build });
+		Params.GetOutDependencies().Emplace(IAssetDependencyGatherer::FGathereredDependency{ FilteredAsset.PackageName, UE::AssetRegistry::EDependencyProperty::Game | UE::AssetRegistry::EDependencyProperty::Build });
 	}
 }
 
 FARFilter FExternalObjectAndActorDependencyGatherer::GetQueryFilter(FName PackageName, TArray<FString>* OutQueryDirectories)
 {
+	TRACE_CPUPROFILER_EVENT_SCOPE(FExternalObjectAndActorDependencyGatherer::GetQueryFilter);
+
 	const FString ExternalActorsPath = ULevel::GetExternalActorsPath(PackageName.ToString());
 	const FString ExternalObjectPath = FExternalPackageHelper::GetExternalObjectsPath(PackageName.ToString());
 	if (OutQueryDirectories)
@@ -38,9 +46,33 @@ FARFilter FExternalObjectAndActorDependencyGatherer::GetQueryFilter(FName Packag
 	Filter.bIncludeOnlyOnDiskAssets = true;
 	Filter.PackagePaths.Add(FName(ExternalActorsPath));
 	Filter.PackagePaths.Add(FName(ExternalObjectPath));
+
+	if (ExternalPathsProvider)
+	{
+		TArray<FName> ExternalPaths = ExternalPathsProvider->GetPathsForPackage(PackageName);
+		for (FName ExternalPath : ExternalPaths)
+		{
+			Filter.PackagePaths.Add(ExternalPath);
+			if (OutQueryDirectories)
+			{
+				OutQueryDirectories->Add(ExternalPath.ToString());
+			}
+		}
+	}
+		
 	return Filter;
 }
 
+void FExternalObjectAndActorDependencyGatherer::SetExternalAssetPathsProvider(IExternalAssetPathsProvider* InProvider)
+{
+	// Can't be set over an already set provider, or we make this a collection
+	check(!InProvider || !ExternalPathsProvider);
+	
+	ExternalPathsProvider = InProvider;
+}
+
+IExternalAssetPathsProvider* FExternalObjectAndActorDependencyGatherer::ExternalPathsProvider = nullptr;
+
 REGISTER_ASSETDEPENDENCY_GATHERER(FExternalObjectAndActorDependencyGatherer, UWorld);
 
 #endif
diff --git a/Engine/Source/Runtime/Engine/Private/HLOD/HLODEngineSubsystem.cpp b/Engine/Source/Runtime/Engine/Private/HLOD/HLODEngineSubsystem.cpp
--- a/Engine/Source/Runtime/Engine/Private/HLOD/HLODEngineSubsystem.cpp
+++ b/Engine/Source/Runtime/Engine/Private/HLOD/HLODEngineSubsystem.cpp
@@ -198,16 +198,10 @@ void UHLODEngineSubsystem::OnPreSaveWorld(UWorld* InWorld, FObjectPreSaveContext
 class FHLODDependencyGatherer : public IAssetDependencyGatherer
 {
 public:
-	virtual void GatherDependencies(const FAssetData& AssetData, const FAssetRegistryState& AssetRegistryState,
-		TFunctionRef<FARCompiledFilter(const FARFilter&)> CompileFilterFunc,
-		TArray<IAssetDependencyGatherer::FGathereredDependency>& OutDependencies,
-		TArray<FString>& OutDependencyDirectories) const override;
+	virtual void GatherDependencies(FGatherDependenciesContext& Context) const override;
 };
 
-void FHLODDependencyGatherer::GatherDependencies(const FAssetData& AssetData,
-	const FAssetRegistryState& AssetRegistryState, TFunctionRef<FARCompiledFilter(const FARFilter&)> CompileFilterFunc,
-	TArray<IAssetDependencyGatherer::FGathereredDependency>& OutDependencies,
-	TArray<FString>& OutDependencyDirectories) const
+void FHLODDependencyGatherer::GatherDependencies(FGatherDependenciesContext& Context) const
 {
 	if (!GetDefault<UHierarchicalLODSettings>()->bSaveLODActorsToHLODPackages)
 	{
@@ -221,18 +215,18 @@ void FHLODDependencyGatherer::GatherDependencies(const FAssetData& AssetData,
 
 	// TODO: GetHLODPackageName is usually constructed from the World's packagename, but this can be replaced by ULevelStreaming::PackageNameToLoad
 	// We need to write the list of LevelStreaming PackageNameToLoad into the AssetData so we can read it from here.
-	FString Wildcard = Utilities->GetWildcardOfHLODPackagesForPackage(AssetData.PackageName.ToString());
+	FString Wildcard = Utilities->GetWildcardOfHLODPackagesForPackage(Context.GetAssetData().PackageName.ToString());
 	FARFilter PossibleAssetsFilter;
 	PossibleAssetsFilter.PackagePaths.Add(FName(FPathViews::GetPath(Wildcard)));
 	TArray<FAssetData> FilteredAssets;
-	AssetRegistryState.GetAssets(CompileFilterFunc(PossibleAssetsFilter), {}, FilteredAssets);
+	Context.GetAssetRegistryState().GetAssets(Context.CompileFilter(PossibleAssetsFilter), {}, FilteredAssets);
 	FString PackageName;
 	for (const FAssetData& FilteredAsset : FilteredAssets)
 	{
 		FilteredAsset.PackageName.ToString(PackageName);
 		if (PackageName.MatchesWildcard(Wildcard))
 		{
-			OutDependencies.Emplace(IAssetDependencyGatherer::FGathereredDependency{ FilteredAsset.PackageName,
+			Context.GetOutDependencies().Emplace(IAssetDependencyGatherer::FGathereredDependency{ FilteredAsset.PackageName,
 				UE::AssetRegistry::EDependencyProperty::Game | UE::AssetRegistry::EDependencyProperty::Build });
 		}
 	}
diff --git a/Engine/Source/Runtime/Engine/Private/WorldPartition/ContentBundle/ContentBundlePaths.cpp b/Engine/Source/Runtime/Engine/Private/WorldPartition/ContentBundle/ContentBundlePaths.cpp
--- a/Engine/Source/Runtime/Engine/Private/WorldPartition/ContentBundle/ContentBundlePaths.cpp
+++ b/Engine/Source/Runtime/Engine/Private/WorldPartition/ContentBundle/ContentBundlePaths.cpp
@@ -8,11 +8,60 @@
 #include "WorldPartition/ContentBundle/ContentBundleBase.h"
 #include "WorldPartition/ContentBundle/ContentBundleDescriptor.h"
 #include "Misc/Paths.h"
+#include "AssetRegistry/AssetData.h"
+#include "ExternalPackageHelper.h"
 
 #if WITH_EDITOR
 #include "PackageTools.h"
 #endif
 
+FName ContentBundlePaths::GetContentBundleGuidsAssetRegistryTag()
+{
+	static FName ContentBundleAssetRegistryTag(ContentBundleAssetRegistryTagString);
+ 
+	return ContentBundleAssetRegistryTag;
+}
+
+void ContentBundlePaths::AddRegistryTags(FAssetRegistryTagsContext Context, TArray<FGuid>& ContentBundleGuids)
+{
+	if (ContentBundleGuids.Num() > 0)
+	{
+		FString ContentBundleGuidsStr = FString::JoinBy(ContentBundleGuids, TEXT(","), [&](const FGuid& Guid) { return Guid.ToString(); });
+		Context.AddTag(UObject::FAssetRegistryTag(ContentBundlePaths::GetContentBundleGuidsAssetRegistryTag(), ContentBundleGuidsStr, UObject::FAssetRegistryTag::TT_Hidden));
+	}
+}
+
+TArray<FGuid> ContentBundlePaths::ParseContentBundleGuids(const FAssetData& AssetData)
+{
+	FString ContentBundleGuidsStr;
+	if (AssetData.GetTagValue(ContentBundlePaths::GetContentBundleGuidsAssetRegistryTag(), ContentBundleGuidsStr))
+	{
+		return ParseContentBundleGuids(ContentBundleGuidsStr);
+	}
+	else
+	{
+		return TArray<FGuid>();
+	}
+}
+
+TArray<FGuid> ContentBundlePaths::ParseContentBundleGuids(const FString& ContentBundleGuidsStr)
+{
+	TArray<FGuid> ContentBundleGuids;
+	TArray<FString> ContentBundleGuidsStrArray;
+
+	ContentBundleGuidsStr.ParseIntoArray(ContentBundleGuidsStrArray, TEXT(","));
+	for (const FString& GuidStr : ContentBundleGuidsStrArray)
+	{
+		FGuid ContentBundleGuid;
+		if (FGuid::Parse(GuidStr, ContentBundleGuid))
+		{
+			ContentBundleGuids.Add(ContentBundleGuid);
+		}
+	}
+
+	return ContentBundleGuids;
+}
+
 FString ContentBundlePaths::GetCookedContentBundleLevelFolder(const FContentBundleBase& ContentBundle)
 {
 	TStringBuilderWithBuffer<TCHAR, NAME_SIZE> GeneratedContentBundleLevelFolder;
@@ -125,6 +174,26 @@ bool ContentBundlePaths::BuildContentBundleExternalActorPath(const FString& InCo
 	return true;
 }
 
+bool ContentBundlePaths::BuildContentBundleActorsRootPath(const FString& InContenBundleMountPoint, const FGuid& InContentBundleGuid, FString& OutContentBundleRootPath)
+{
+	if (InContenBundleMountPoint.IsEmpty() || !InContentBundleGuid.IsValid())
+	{
+		return false;
+	}
+
+	TStringBuilderWithBuffer<TCHAR, NAME_SIZE> ContentBundleRootPathBuilder;
+	ContentBundleRootPathBuilder += TEXT("/");
+	ContentBundleRootPathBuilder += InContenBundleMountPoint;
+	ContentBundleRootPathBuilder += TEXT("/");
+	ContentBundleRootPathBuilder += FPackagePath::GetExternalActorsFolderName();
+	ContentBundleRootPathBuilder += GetContentBundleFolder();
+	ContentBundleRootPathBuilder += InContentBundleGuid.ToString();
+
+	OutContentBundleRootPath = *ContentBundleRootPathBuilder;
+
+	return true;
+}
+
 bool ContentBundlePaths::BuildActorDescContainerPackagePath(const FString& InContenBundleMountPoint, const FGuid& InContentBundleGuid, const FString& InLevelPackagePath, FString& OutContainerPackagePath)
 {
 	FString ContentBundleRootPath;
@@ -168,4 +237,43 @@ bool ContentBundlePaths::IsAContentBundlePath(FStringView InContentBundlePath)
 	return GetRelativePath(InContentBundlePath).Len() > 0;
 }
 
+void ContentBundlePaths::ForEachContentBundleLevelPackagePath(const FString& InLevelPackageName, TFunctionRef<void(const FString&)> Func)
+{
+	TRACE_CPUPROFILER_EVENT_SCOPE(ContentBundlePaths::ForEachContentBundleLevelPackagePath);
+
+	UClass* GameFeatureDataClass = FindObject<UClass>(nullptr, TEXT("/Script/GameFeatures.GameFeatureData"));
+	if (GameFeatureDataClass)
+	{
+		TArray<FAssetData> AssetsData;
+
+		FARFilter Filter;
+		Filter.bIncludeOnlyOnDiskAssets = false;
+		Filter.ClassPaths = { GameFeatureDataClass->GetClassPathName() };
+		Filter.bRecursiveClasses = true;
+		Filter.bIncludeOnlyOnDiskAssets = IsRunningCookCommandlet();	// When cooking, only enumerate on-disk assets- avoids indeterministic results
+		FExternalPackageHelper::GetSortedAssets(Filter, AssetsData);
+		
+		for (const FAssetData& AssetData : AssetsData)
+		{
+			TArray<FGuid> ContentBundleGuids;
+			ContentBundleGuids = ParseContentBundleGuids(AssetData);
+
+			if (ContentBundleGuids.Num())
+			{
+				const FString MountPoint = FPackageName::GetPackageMountPoint(AssetData.PackagePath.ToString()).ToString();
+
+				for (const FGuid& ContentBundleGuid : ContentBundleGuids)
+				{
+					FString ContentBundleExternalActorPath;
+					if (ContentBundlePaths::BuildActorDescContainerPackagePath(MountPoint, ContentBundleGuid, InLevelPackageName, ContentBundleExternalActorPath))
+					{
+						Func(ContentBundleExternalActorPath);
+					}
+				}
+			}
+		}
+	}
+}
+
+
 #endif
diff --git a/Engine/Source/Runtime/Engine/Private/WorldPartition/DataLayer/ExternalDataLayerHelper.cpp b/Engine/Source/Runtime/Engine/Private/WorldPartition/DataLayer/ExternalDataLayerHelper.cpp
--- a/Engine/Source/Runtime/Engine/Private/WorldPartition/DataLayer/ExternalDataLayerHelper.cpp
+++ b/Engine/Source/Runtime/Engine/Private/WorldPartition/DataLayer/ExternalDataLayerHelper.cpp
@@ -42,6 +42,24 @@ bool FExternalDataLayerHelper::BuildExternalDataLayerRootPath(const FString& InE
 	return true;
 }
 
+bool FExternalDataLayerHelper::BuildExternalDataLayerActorsRootPath(const FString& InEDLMountPoint, const FExternalDataLayerUID& InExternalDataLayerUID, FString& OutExternalDataLayerRootPath)
+{
+	if (InEDLMountPoint.IsEmpty() || !InExternalDataLayerUID.IsValid())
+	{
+		return false;
+	}
+
+	TStringBuilderWithBuffer<TCHAR, NAME_SIZE> Builder;
+	Builder += TEXT("/");
+	Builder += InEDLMountPoint;
+	Builder += TEXT("/");
+	Builder += FPackagePath::GetExternalActorsFolderName();
+	Builder += GetExternalDataLayerFolder();
+	Builder += InExternalDataLayerUID.ToString();
+	OutExternalDataLayerRootPath = *Builder;
+	return true;
+}
+
 FString FExternalDataLayerHelper::GetExternalDataLayerLevelRootPath(const UExternalDataLayerAsset* InExternalDataLayerAsset, const FString& InLevelPackagePath)
 {
 	check(InExternalDataLayerAsset);
@@ -379,9 +397,8 @@ void FExternalDataLayerHelper::ForEachExternalDataLayerLevelPackagePath(const FS
 	if (GameFeatureDataClass)
 	{
 		FARFilter Filter;
-		Filter.bIncludeOnlyOnDiskAssets = false;
 		Filter.ClassPaths = { GameFeatureDataClass->GetClassPathName() };
-		Filter.bRecursivePaths = true;
+		Filter.bRecursiveClasses = true;
 		Filter.bIncludeOnlyOnDiskAssets = IsRunningCookCommandlet();	// When cooking, only enumerate on-disk assets- avoids indeterministic results
 		TArray<FAssetData> AssetsData;
 		FExternalPackageHelper::GetSortedAssets(Filter, AssetsData);
diff --git a/Engine/Source/Runtime/Engine/Public/WorldPartition/ContentBundle/ContentBundlePaths.h b/Engine/Source/Runtime/Engine/Public/WorldPartition/ContentBundle/ContentBundlePaths.h
--- a/Engine/Source/Runtime/Engine/Public/WorldPartition/ContentBundle/ContentBundlePaths.h
+++ b/Engine/Source/Runtime/Engine/Public/WorldPartition/ContentBundle/ContentBundlePaths.h
@@ -3,13 +3,14 @@
 
 #include "CoreMinimal.h"
 #include "Containers/StringView.h"
-
+#include "UObject/AssetRegistryTagsContext.h"
 class FContentBundleBase;
 
 namespace ContentBundlePaths
 {
 	inline constexpr FStringView ContentBundleFolder = TEXTVIEW("/ContentBundle/");
 	inline constexpr FStringView GeneratedFolder = TEXTVIEW("/_Generated_/");
+	inline constexpr FStringView ContentBundleAssetRegistryTagString = TEXTVIEW("ContentBundleGuids");	
 
 	// Returns "/ContentBundle/"
 	constexpr FStringView GetContentBundleFolder() { return ContentBundleFolder; }
@@ -31,6 +32,12 @@ namespace ContentBundlePaths
 		return FStringView(&GetGeneratedFolder()[1], GetGeneratedFolder().Len() - 2);
 	}
 
+	ENGINE_API FName GetContentBundleGuidsAssetRegistryTag();
+	ENGINE_API void AddRegistryTags(FAssetRegistryTagsContext Context, TArray<FGuid>& ContentBundleGuids);
+	ENGINE_API TArray<FGuid> ParseContentBundleGuids(const FString& GuidsStr);
+	ENGINE_API TArray<FGuid> ParseContentBundleGuids(const FAssetData& AssetData);
+	ENGINE_API void ForEachContentBundleLevelPackagePath(const FString& InLevelPackageName, TFunctionRef<void(const FString&)> Func);
+
 	// Return value uses a reduced amount of character to avoid reaching character limit when cooking 
 	// return format is /{MountPoint}/CB/{ContentBundle_ShortUID}/{LevelPath}/{GeneratedFolder}/
 	FString GetCookedContentBundleLevelFolder(const FContentBundleBase& ContentBundle);
@@ -42,6 +49,9 @@ namespace ContentBundlePaths
 	// return format is /{MountPoint}/{ContentBundleFolder}/{ContentBundleUID}/
 	ENGINE_API bool BuildContentBundleExternalActorPath(const FString& InContenBundleMountPoint, const FGuid& InContentBundleGuid, FString& OutContentBundleRootPath);
 
+	// return format is /{MountPoint}/{ExternalActorFolder}/{ContentBundleFolder}/{ContentBundleUID}
+	ENGINE_API bool BuildContentBundleActorsRootPath(const FString& InContenBundleMountPoint, const FGuid& InContentBundleGuid,  FString& OutContainerPackagePath);
+
 	// return format is /{MountPoint}/{ContentBundleFolder}/{ContentBundleUID}/{LevelPath}/
 	ENGINE_API bool BuildActorDescContainerPackagePath(const FString& InContenBundleMountPoint, const FGuid& InContentBundleGuid, const FString& InLevelPackagePath, FString& OutContainerPackagePath);
 
@@ -55,7 +65,7 @@ namespace ContentBundlePaths
 	// return an ExternalActor path following the format : /{MountPoint}/{ExternalActorFolder}/{ContentBundleFolder}/{ContentBundleUID}/{LevelPath}/{ExternalActorPackagePath}
 	ENGINE_API FString MakeExternalActorPackagePath(const FString& ContentBundleExternalActorFolder, const FString& ActorName);
 
-	// return true if InPackage follow format : //{MountPoint}/{ExternalActorFolder}/{ContentBundleFolder}/{ContentBundleUID}/*, false otherwise
+	// return true if InPackage follow format : /{MountPoint}/{ExternalActorFolder}/{ContentBundleFolder}/{ContentBundleUID}/*, false otherwise
 	ENGINE_API bool IsAContentBundleExternalActorPackagePath(FStringView InPackagePath);
 
 	// InContentBundleExternalActorPackagePath format is : /{MountPoint}/{ExternalActorFolder}/{ContentBundleFolder}/{ContentBundleUID}/{LevelPath}/{ExternalActorPackagePath}
diff --git a/Engine/Source/Runtime/Engine/Public/WorldPartition/DataLayer/ExternalDataLayerHelper.h b/Engine/Source/Runtime/Engine/Public/WorldPartition/DataLayer/ExternalDataLayerHelper.h
--- a/Engine/Source/Runtime/Engine/Public/WorldPartition/DataLayer/ExternalDataLayerHelper.h
+++ b/Engine/Source/Runtime/Engine/Public/WorldPartition/DataLayer/ExternalDataLayerHelper.h
@@ -39,12 +39,17 @@ public:
 	 */
 	ENGINE_API static bool BuildExternalDataLayerRootPath(const FString& InEDLMountPoint, const FExternalDataLayerUID& InExternalDataLayerUID, FString& OutExternalDataLayerRootPath);
 
-	/** Return true if succeeds building the external data layer level root path (OutExternalDataLayerLevelRootPath) using the ExternalDataLayerAsset and the level package path.
+	/** Return true if succeeds building the external data layer actors root path (OutExternalDataLayerRootPath) using the provided mount point and EDL UID.
+	 * Format is /{MountPoint}/{ExternalActorFolder}/{ExternalDataLayerFolder}/{EDL_UID}
+	 */
+	ENGINE_API static bool BuildExternalDataLayerActorsRootPath(const FString& InEDLMountPoint, const FExternalDataLayerUID& InExternalDataLayerUID, FString& OutExternalDataLayerRootPath);
+
+	/** Return the external data layer level root path (OutExternalDataLayerLevelRootPath) using the ExternalDataLayerAsset and the level package path.
 	  * Format is /{MountPoint}/{ExternalDataLayerFolder}/{EDL_UID}/{LevelPath}
 	  */
 	ENGINE_API static FString GetExternalDataLayerLevelRootPath(const UExternalDataLayerAsset* InExternalDataLayerAsset, const FString& InLevelPackagePath);
 
-	/** Return true if succeeds building the external data layer level root path (OutExternalDataLayerLevelRootPath) using the provided mount point, EDL UID and the level package path.
+	/** Return the external data layer level root path (OutExternalDataLayerLevelRootPath) using the provided mount point, EDL UID and the level package path.
 	  * Format is /{MountPoint}/{ExternalDataLayerFolder}/{EDL_UID}/{LevelPath}
 	  */
 	ENGINE_API static FString GetExternalDataLayerLevelRootPath(const FString& InEDLMountPoint, const FExternalDataLayerUID& InExternalDataLayerUID, const FString& InLevelPackagePath);
