Remove ensure(bCacheIsUpToDate) we can't really guarantee that from all call paths right now.,
Fix for IncrementalCook.

diff --git a/Engine/Plugins/Runtime/GameFeatures/Source/GameFeatures/Private/GameFeatureData.cpp b/Engine/Plugins/Runtime/GameFeatures/Source/GameFeatures/Private/GameFeatureData.cpp
--- a/Engine/Plugins/Runtime/GameFeatures/Source/GameFeatures/Private/GameFeatureData.cpp
+++ b/Engine/Plugins/Runtime/GameFeatures/Source/GameFeatures/Private/GameFeatureData.cpp
@@ -1039,7 +1039,6 @@ void FGameFeatureDataExternalAssetsPathCache::UpdateCache(const FUpdateCacheCont
 
 TArray<FName> FGameFeatureDataExternalAssetsPathCache::GetPathsForPackage(FName LevelPath)
 {
-	ensure(bCacheIsUpToDate);
 	TArray<FName> Paths;
 	PerLevelAssetDirectories.MultiFind(LevelPath, Paths);
 	return Paths;
