[GFP] Pass Keeplist for UGameFeatureOptionalContentInstaller

#tests CI, local build
#rnx

diff --git a/Engine/Plugins/Runtime/GameFeatures/Source/GameFeatures/Private/GameFeatureOptionalContentInstaller.cpp b/Engine/Plugins/Runtime/GameFeatures/Source/GameFeatures/Private/GameFeatureOptionalContentInstaller.cpp
--- a/Engine/Plugins/Runtime/GameFeatures/Source/GameFeatures/Private/GameFeatureOptionalContentInstaller.cpp
+++ b/Engine/Plugins/Runtime/GameFeatures/Source/GameFeatures/Private/GameFeatureOptionalContentInstaller.cpp
@@ -33,9 +33,12 @@ void UGameFeatureOptionalContentInstaller::BeginDestroy()
 	Super::BeginDestroy();
 }
 
-void UGameFeatureOptionalContentInstaller::Init(TUniqueFunction<TArray<FName>(FString)> InGetOptionalBundlePredicate)
+void UGameFeatureOptionalContentInstaller::Init(
+	TUniqueFunction<TArray<FName>(FString)> InGetOptionalBundlePredicate,
+	TUniqueFunction<TArray<FName>(FString)> InGetOptionalKeepBundlePredicate)
 {
 	GetOptionalBundlePredicate = MoveTemp(InGetOptionalBundlePredicate);
+	GetOptionalKeepBundlePredicate = MoveTemp(InGetOptionalKeepBundlePredicate);
 	BundleManager = IInstallBundleManager::GetPlatformInstallBundleManager();
 
 	// Create the cvar sink
@@ -297,13 +300,14 @@ void UGameFeatureOptionalContentInstaller::ReleaseContent(const FString& PluginN
 	{
 		return;
 	}
+	TArray<FName> KeepBundles = GetOptionalKeepBundlePredicate(PluginName);
 
 	Flags |= EInstallBundleReleaseRequestFlags::ExplicitRemoveList;
 
 	BundleManager->RequestReleaseContent(
 		Bundles, 
 		Flags,
-		{}, 
+		KeepBundles, 
 		GameFeatureOptionalContentInstaller::InstallBundleManagerVerbosityOverride);
 }
 
diff --git a/Engine/Plugins/Runtime/GameFeatures/Source/GameFeatures/Public/GameFeatureOptionalContentInstaller.h b/Engine/Plugins/Runtime/GameFeatures/Source/GameFeatures/Public/GameFeatureOptionalContentInstaller.h
--- a/Engine/Plugins/Runtime/GameFeatures/Source/GameFeatures/Public/GameFeatureOptionalContentInstaller.h
+++ b/Engine/Plugins/Runtime/GameFeatures/Source/GameFeatures/Public/GameFeatureOptionalContentInstaller.h
@@ -33,7 +33,9 @@ public:
 public:
 	virtual void BeginDestroy() override;
 
-	void GAMEFEATURES_API Init(TUniqueFunction<TArray<FName>(FString)> GetOptionalBundlePredicate);
+	void GAMEFEATURES_API Init(
+		TUniqueFunction<TArray<FName>(FString)> GetOptionalBundlePredicate,
+		TUniqueFunction<TArray<FName>(FString)> GetOptionalKeepBundlePredicate = [](auto) { return TArray<FName>{}; });
 
 	void GAMEFEATURES_API Enable(bool bEnable);
 
@@ -81,6 +83,7 @@ private:
 
 private:
 	TUniqueFunction<TArray<FName>(FString)> GetOptionalBundlePredicate;
+	TUniqueFunction<TArray<FName>(FString)> GetOptionalKeepBundlePredicate;
 	TSharedPtr<IInstallBundleManager> BundleManager;
 	TSet<FString> RelevantGFPs;
 	TMap<FString, FGFPInstall> ActiveGFPInstalls;
