Print warning for dlc plugins using global core redirects
- FRedirectionSummary : add new delegate called when a global redirect is added.
- FCoreRedirects : add declaration of new delegate FGlobalCoreRedirectAddedEvent. That's the delegate used in FRedirectionSummary. I put it here because FRedirectionSummary is a private header and I need it to be in a public header so other system can attach to it. Also add a static getter OnGlobalCoreRedirectAdded to retrieve the delegate from FRedirectionSummary.
- UGameFeaturesSubsystem : new static variable used to store the names of the plugins specific to a dlc.
- GameFeatureData : When ini files of plugins are loaded and if the plugin is a dlc plugin, attach to the event OnGlobalCoreRedirectAdded and print a warning. I made it a RAII class to attach/remove the callback automatically. It's wrapped in a precompile if because it's editor only, the game version doesn't do anything.

#rnx

diff --git a/Engine/Plugins/Runtime/GameFeatures/Source/GameFeatures/Private/GameFeatureData.cpp b/Engine/Plugins/Runtime/GameFeatures/Source/GameFeatures/Private/GameFeatureData.cpp
--- a/Engine/Plugins/Runtime/GameFeatures/Source/GameFeatures/Private/GameFeatureData.cpp
+++ b/Engine/Plugins/Runtime/GameFeatures/Source/GameFeatures/Private/GameFeatureData.cpp
@@ -83,6 +83,51 @@ static TAutoConsoleVariable<bool> CVarAllowRuntimeDeviceProfiles(
 	TEXT("Allow game feature plugins to generate device profiles from config based on existing parents"),
 	ECVF_Default);
 
+namespace GameFeatureData
+{
+#if WITH_EDITOR
+	class FScopeGlobalCoreRedirectAdded
+	{
+	public:
+		FScopeGlobalCoreRedirectAdded(const FString& PluginName, const FString& PluginConfigFilename)
+		{
+			if (IsRunningDLCCookCommandlet() && UGameFeaturesSubsystem::GetDlcPlugins().Contains(PluginName))
+			{
+				FCoreRedirects::FGlobalCoreRedirectAddedEvent* Event = FCoreRedirects::GetOnGlobalCoreRedirectAdded();
+				if (!Event)
+				{
+					return;
+				}
+
+				OnGlobalCoreRedirectAddedHandle = Event->AddLambda([&PluginName, &PluginConfigFilename](const FCoreRedirect& CoreRedirect)
+					{
+						UE_LOG(LogGameFeatures, Warning, TEXT("Plugin %s added global core redirect defined in file %s (%s -> %s). This will trigger useless recooks."),
+							*PluginName, *PluginConfigFilename, *CoreRedirect.OldName.ToString(), *CoreRedirect.NewName.ToString());
+					});
+			}
+		}
+
+		~FScopeGlobalCoreRedirectAdded()
+		{
+			if (OnGlobalCoreRedirectAddedHandle.IsValid())
+			{
+				FCoreRedirects::GetOnGlobalCoreRedirectAdded()->Remove(OnGlobalCoreRedirectAddedHandle);
+			}
+		}
+
+	private:
+		FDelegateHandle OnGlobalCoreRedirectAddedHandle;
+	};
+#else // #if WITH_EDITOR
+	class FScopeGlobalCoreRedirectAdded
+	{
+	public:
+		FScopeGlobalCoreRedirectAdded(const FString& PluginName, const FString& PluginConfigFilename) {}
+		~FScopeGlobalCoreRedirectAdded(){}
+	};
+#endif // #if WITH_EDITOR
+}
+
 void UGameFeatureData::InitializeBasePluginIniFile(const FString& PluginInstalledFilename)
 {
 	const FString PluginName = FPaths::GetBaseFilename(PluginInstalledFilename);
@@ -117,6 +162,8 @@ void UGameFeatureData::InitializeBasePluginIniFile(const FString& PluginInstalle
 				FConfigCacheIni::RegisterPlugin(*Plugin->GetName(), Plugin->GetBaseDir(), Plugin->GetExtensionBaseDirs(), DynamicLayerPriority::GameFeature, bIncludePluginNameInBranchName);
 			}
 
+			GameFeatureData::FScopeGlobalCoreRedirectAdded ScopeWarnOnDlcAddingGlobalCoreRedirect(PluginName, PluginConfigFilename);
+
 			FCoreRedirects::ReadRedirectsFromIni(PluginConfigFilename);
 			ReloadConfigs(PluginConfig);
 
@@ -147,6 +194,8 @@ void UGameFeatureData::InitializeBasePluginIniFile(const FString& PluginInstalle
 			ReloadConfigs(*PluginConfig);
 		}
 
+		GameFeatureData::FScopeGlobalCoreRedirectAdded ScopeWarnOnDlcAddingGlobalCoreRedirect(PluginName, PluginConfigFilename);
+
 		// @todo move this into ReloadObjectsFromModifiedConfigSections?
 		FCoreRedirects::ReadRedirectsFromIni(PluginConfigFilename);
 
@@ -167,6 +216,7 @@ void UGameFeatureData::InitializeBasePluginIniFile(const FString& PluginInstalle
 		}
 		else
 		{
+			GameFeatureData::FScopeGlobalCoreRedirectAdded ScopeWarnOnDlcAddingGlobalCoreRedirect(PluginName, PluginConfigFilename);
 			FCoreRedirects::ReadRedirectsFromIni(PluginConfigFilename);
 			ReloadConfigs(PluginConfig);
 		}
diff --git a/Engine/Plugins/Runtime/GameFeatures/Source/GameFeatures/Private/GameFeaturesSubsystem.cpp b/Engine/Plugins/Runtime/GameFeatures/Source/GameFeatures/Private/GameFeaturesSubsystem.cpp
--- a/Engine/Plugins/Runtime/GameFeatures/Source/GameFeatures/Private/GameFeaturesSubsystem.cpp
+++ b/Engine/Plugins/Runtime/GameFeatures/Source/GameFeatures/Private/GameFeaturesSubsystem.cpp
@@ -276,6 +276,8 @@ namespace UE::GameFeatures
 #endif
 }
 
+TSet<FString> UGameFeaturesSubsystem::DlcPlugins;
+
 const FString LexToString(const EBuiltInAutoState BuiltInAutoState)
 {
 	switch (BuiltInAutoState)
@@ -4522,12 +4524,22 @@ void UGameFeaturesSubsystem::GetPluginsToCook(TSet<FString>& OutPlugins)
 			}
 		}
 
+		if (IsRunningDLCCookCommandlet())
+		{
+			DlcPlugins.Append(ReturnList);
+		}
+
 		return ReturnList;
 	}();
 	
 	OutPlugins.Append(PluginsList);	
 }
 
+const TSet<FString>& UGameFeaturesSubsystem::GetDlcPlugins()
+{
+	return DlcPlugins;
+}
+
 bool UGameFeaturesSubsystem::GetPluginDebugStateEnabled(const FString& PluginUrl)
 {
 #if !UE_BUILD_SHIPPING
diff --git a/Engine/Plugins/Runtime/GameFeatures/Source/GameFeatures/Public/GameFeaturesSubsystem.h b/Engine/Plugins/Runtime/GameFeatures/Source/GameFeatures/Public/GameFeaturesSubsystem.h
--- a/Engine/Plugins/Runtime/GameFeatures/Source/GameFeatures/Public/GameFeaturesSubsystem.h
+++ b/Engine/Plugins/Runtime/GameFeatures/Source/GameFeatures/Public/GameFeaturesSubsystem.h
@@ -726,6 +726,8 @@ public:
 	/** Used during a DLC cook to determine which plugins should be cooked */
 	static UE_API void GetPluginsToCook(TSet<FString>& OutPlugins);
 
+	static UE_API const TSet<FString>& GetDlcPlugins();
+
 	UE_API bool GetPluginDebugStateEnabled(const FString& PluginUrl);
 	UE_API void SetPluginDebugStateEnabled(const FString& PluginUrl, bool bEnabled);
 
@@ -945,6 +947,8 @@ private:
 	FDelegateHandle GetExplanationForUnavailablePackageDelegateHandle;
 
 	bool bInitializedPolicyManager = false;
+
+	static UE_API TSet<FString> DlcPlugins;
 };
 
 #undef UE_API
diff --git a/Engine/Source/Runtime/CoreUObject/Internal/UObject/CoreRedirects/RedirectionSummary.h b/Engine/Source/Runtime/CoreUObject/Internal/UObject/CoreRedirects/RedirectionSummary.h
--- a/Engine/Source/Runtime/CoreUObject/Internal/UObject/CoreRedirects/RedirectionSummary.h
+++ b/Engine/Source/Runtime/CoreUObject/Internal/UObject/CoreRedirects/RedirectionSummary.h
@@ -33,6 +33,8 @@ public:
 	void GetHashAffectingPackages(const TConstArrayView<FName>& PackageNames, TArray<FBlake3Hash>& HashArray);
 	void AppendHashGlobal(FBlake3& Hasher);
 
+	FCoreRedirects::FGlobalCoreRedirectAddedEvent& GetOnGlobalCoreRedirectAdded();
+
 private:
 	struct FCompareRedirect
 	{
@@ -71,6 +73,8 @@ private:
 	 * them dirty does not cause a data race with the second thread requests them.
 	 */
 	FRWLock Lock;
+
+	FCoreRedirects::FGlobalCoreRedirectAddedEvent OnGlobalCoreRedirectAdded;
 };
 
 #endif
\ No newline at end of file
diff --git a/Engine/Source/Runtime/CoreUObject/Private/UObject/CoreRedirects.cpp b/Engine/Source/Runtime/CoreUObject/Private/UObject/CoreRedirects.cpp
--- a/Engine/Source/Runtime/CoreUObject/Private/UObject/CoreRedirects.cpp
+++ b/Engine/Source/Runtime/CoreUObject/Private/UObject/CoreRedirects.cpp
@@ -2062,7 +2062,17 @@ void FCoreRedirects::RecordRemovedObjectRedirector(const FSoftObjectPath& Source
 	}
 }
 
-#endif
+FCoreRedirects::FGlobalCoreRedirectAddedEvent* FCoreRedirects::GetOnGlobalCoreRedirectAdded()
+{
+	FCoreRedirectsContext& Context = FCoreRedirectsContext::GetThreadContext();
+	if (!!(Context.GetFlags() & FCoreRedirectsContext::EFlags::UseRedirectionSummary))
+	{
+		return &Context.GetRedirectionSummary().GetOnGlobalCoreRedirectAdded();
+	}
+	return nullptr;
+}
+
+#endif // #if WITH_EDITOR
 
 bool FCoreRedirects::RunTests()
 {
diff --git a/Engine/Source/Runtime/CoreUObject/Private/UObject/RedirectionSummary.cpp b/Engine/Source/Runtime/CoreUObject/Private/UObject/RedirectionSummary.cpp
--- a/Engine/Source/Runtime/CoreUObject/Private/UObject/RedirectionSummary.cpp
+++ b/Engine/Source/Runtime/CoreUObject/Private/UObject/RedirectionSummary.cpp
@@ -55,6 +55,8 @@ void FRedirectionSummary::Add(const FCoreRedirect& CoreRedirect)
 	else if (CoreRedirect.IsWildcardMatch() || CoreRedirect.OldName.PackageName.IsNone())
 	{
 		GlobalRedirects.Add(FCoreRedirect(CoreRedirect));
+
+		OnGlobalCoreRedirectAdded.Broadcast(CoreRedirect);
 	}
 	else
 	{
@@ -205,6 +207,11 @@ void FRedirectionSummary::AppendHashGlobal(FBlake3& Hasher)
 	GlobalRedirects.AppendHashInWriteLock(Hasher);
 }
 
+FCoreRedirects::FGlobalCoreRedirectAddedEvent& FRedirectionSummary::GetOnGlobalCoreRedirectAdded()
+{
+	return OnGlobalCoreRedirectAdded;
+}
+
 void FRedirectionSummary::FRedirectContainer::Add(FCoreRedirect&& Redirect)
 {
 	Redirects.Add(MoveTemp(Redirect));
diff --git a/Engine/Source/Runtime/CoreUObject/Public/UObject/CoreRedirects.h b/Engine/Source/Runtime/CoreUObject/Public/UObject/CoreRedirects.h
--- a/Engine/Source/Runtime/CoreUObject/Public/UObject/CoreRedirects.h
+++ b/Engine/Source/Runtime/CoreUObject/Public/UObject/CoreRedirects.h
@@ -423,6 +423,16 @@ struct FCoreRedirects
 	static COREUOBJECT_API void RecordAddedObjectRedirector(const FSoftObjectPath& Source, const FSoftObjectPath& Dest);
 	/** Remove the given Source->Path redirector to the summary used for AppendHashOfRedirectsAffectingPackages. */
 	static COREUOBJECT_API void RecordRemovedObjectRedirector(const FSoftObjectPath& Source, const FSoftObjectPath& Dest);
+
+	/** Thread-safe delegate broadcasted when a global core redirect is added to FRedirectionSummary. It can be called from any thread. */
+	DECLARE_TS_MULTICAST_DELEGATE_OneParam(FGlobalCoreRedirectAddedEvent, const FCoreRedirect& CoreRedirect);
+	
+	/** 
+	 * Return the delegate called when a global core redirect is added to FRedirectionSummary.
+	 * Can return nullptr if the current thread context does not use a FRedirectionSummary.
+	 */
+	static COREUOBJECT_API FGlobalCoreRedirectAddedEvent* GetOnGlobalCoreRedirectAdded();
+
 #endif
 
 	/** Runs set of redirector tests, returns false on failure */
