Change log level from Warning to Log

#rnx

diff --git a/Engine/Plugins/Runtime/GameFeatures/Source/GameFeatures/Private/GameFeatureData.cpp b/Engine/Plugins/Runtime/GameFeatures/Source/GameFeatures/Private/GameFeatureData.cpp
--- a/Engine/Plugins/Runtime/GameFeatures/Source/GameFeatures/Private/GameFeatureData.cpp
+++ b/Engine/Plugins/Runtime/GameFeatures/Source/GameFeatures/Private/GameFeatureData.cpp
@@ -101,7 +101,7 @@ namespace GameFeatureData
 
 				OnGlobalCoreRedirectAddedHandle = Event->AddLambda([&PluginName, &PluginConfigFilename](const FCoreRedirect& CoreRedirect)
 					{
-						UE_LOG(LogGameFeatures, Warning, TEXT("Plugin %s added global core redirect defined in file %s (%s -> %s). This will trigger useless recooks."),
+						UE_LOG(LogGameFeatures, Log, TEXT("Plugin %s added global core redirect defined in file %s (%s -> %s). This will trigger useless recooks."),
 							*PluginName, *PluginConfigFilename, *CoreRedirect.OldName.ToString(), *CoreRedirect.NewName.ToString());
 					});
 			}
