In a multi product cook, no plugin should add a global core redirect. Otherwise it's gonna change the global dependencies hash:

- if a base game plugin adds a global core redirect : a dlc cook might not load it.
- if a dlc plugin adds a global core redirect : another dlc cook will not load it.

So display a message for all the global core redirects.

#rnx

diff --git a/Engine/Plugins/Runtime/GameFeatures/Source/GameFeatures/Private/GameFeatureData.cpp b/Engine/Plugins/Runtime/GameFeatures/Source/GameFeatures/Private/GameFeatureData.cpp
--- a/Engine/Plugins/Runtime/GameFeatures/Source/GameFeatures/Private/GameFeatureData.cpp
+++ b/Engine/Plugins/Runtime/GameFeatures/Source/GameFeatures/Private/GameFeatureData.cpp
@@ -91,20 +91,22 @@ namespace GameFeatureData
 	public:
 		FScopeGlobalCoreRedirectAdded(const FString& PluginName, const FString& PluginConfigFilename)
 		{
-			if (IsRunningDLCCookCommandlet() && UGameFeaturesSubsystem::GetDlcPlugins().Contains(PluginName))
+			if (!IsRunningCookCommandlet())
 			{
-				FCoreRedirects::FGlobalCoreRedirectAddedEvent* Event = FCoreRedirects::GetOnGlobalCoreRedirectAdded();
-				if (!Event)
-				{
-					return;
-				}
+				return;
+			}
 
-				OnGlobalCoreRedirectAddedHandle = Event->AddLambda([&PluginName, &PluginConfigFilename](const FCoreRedirect& CoreRedirect)
-					{
-						UE_LOG(LogGameFeatures, Log, TEXT("Plugin %s added global core redirect defined in file %s (%s -> %s). This will trigger useless recooks."),
-							*PluginName, *PluginConfigFilename, *CoreRedirect.OldName.ToString(), *CoreRedirect.NewName.ToString());
-					});
+			FCoreRedirects::FGlobalCoreRedirectAddedEvent* Event = FCoreRedirects::GetOnGlobalCoreRedirectAdded();
+			if (!Event)
+			{
+				return;
 			}
+
+			OnGlobalCoreRedirectAddedHandle = Event->AddLambda([&PluginName, &PluginConfigFilename](const FCoreRedirect& CoreRedirect)
+				{
+					UE_LOG(LogGameFeatures, Display, TEXT("Plugin %s added global core redirect defined in file %s (%s -> %s). This will trigger useless recooks."),
+						*PluginName, *PluginConfigFilename, *CoreRedirect.OldName.ToString(), *CoreRedirect.NewName.ToString());
+				});
 		}
 
 		~FScopeGlobalCoreRedirectAdded()
