- Renamed [ConoleVariables] and [ConsoleVariables_HighPriority] to [PluginCVars_LowPriority] and [PluginCVars_HighPriority]
- Added initial support for them in Preview, although it only works if in Engine.ini files, a future change will work to address this

#FailedValidations Preflight

diff --git a/Engine/Plugins/Runtime/GameFeatures/Source/GameFeatures/Private/GameFeatureData.cpp b/Engine/Plugins/Runtime/GameFeatures/Source/GameFeatures/Private/GameFeatureData.cpp
--- a/Engine/Plugins/Runtime/GameFeatures/Source/GameFeatures/Private/GameFeatureData.cpp
+++ b/Engine/Plugins/Runtime/GameFeatures/Source/GameFeatures/Private/GameFeatureData.cpp
@@ -239,19 +239,27 @@ void UGameFeatureData::InitializeHierarchicalPluginIniFiles(const FString& Plugi
 		UE::DynamicConfig::PerformDynamicConfig(PluginName, [PluginName](FConfigModificationTracker* ChangeTracker)
 			{
 				// set which sections to track for cvars, with their priority
-				ChangeTracker->CVars.Add(TEXT("ConsoleVariables")).CVarPriority = (int)ECVF_SetByPluginLowPriority;
-				ChangeTracker->CVars.Add(TEXT("ConsoleVariables_HighPriority")).CVarPriority = (int)ECVF_SetByPluginHighPriority;
+				ChangeTracker->CVars.Add(TEXT("PluginCVars_LowPriority")).CVarPriority = (int)ECVF_SetByPluginLowPriority;
+				ChangeTracker->CVars.Add(TEXT("PluginCVars_HighPriority")).CVarPriority = (int)ECVF_SetByPluginHighPriority;
 
 				// apply plugin modifications from this plugin to the everything
 				FConfigCacheIni::AddPluginToAllBranches(PluginName, ChangeTracker);
+
+#if !UE_BUILD_SHIPPING
+				for (const TPair<FName, TSet<FString>>& Pair : ChangeTracker->ModifiedSectionsPerBranch)
+				{
+					ensureAlwaysMsgf(!Pair.Value.Contains(TEXT("ConsoleVariables")) && !Pair.Value.Contains(TEXT("ConsoleVariables_HighPriority")),
+						TEXT("GameFeaturePlugins .ini files should no longer use [ConsoleVariables] or [ConsoleVariables_HighPriority] sections. Instead, use [PluginCVars_LowPriority] and [PluginCVars_HighPriority]."));
+				}
+#endif
 			});
 
 		UE::DynamicConfig::PerformDynamicConfig(PluginName, [PluginName](FConfigModificationTracker* ChangeTracker)
 			{
 				{
 					// set which sections to track for cvars, with their priority
-					ChangeTracker->CVars.Add(TEXT("ConsoleVariables")).CVarPriority = (int)ECVF_SetByHotfix;
-					ChangeTracker->CVars.Add(TEXT("ConsoleVariables_HighPriority")).CVarPriority = (int)ECVF_SetByHotfix;
+					ChangeTracker->CVars.Add(TEXT("PluginCVars_LowPriority")).CVarPriority = (int)ECVF_SetByHotfix;
+					ChangeTracker->CVars.Add(TEXT("PluginCVars_HighPriority")).CVarPriority = (int)ECVF_SetByHotfix;
 
 					// give hotfix a chance to run at that point
 #define HOTFIX_BRANCH(Ini) UE::DynamicConfig::HotfixPluginForBranch.Broadcast(PluginName, #Ini, ChangeTracker);
@@ -656,19 +664,27 @@ void UGameFeatureData::InitializeHierarchicalPluginIniFiles(const TArrayView<FSt
 		UE::DynamicConfig::PerformDynamicConfig("InitializeHierarchicalPluginIniFiles", [&PluginNames](FConfigModificationTracker* ChangeTracker)
 			{
 				// set which sections to track for cvars, with their priority
-				ChangeTracker->CVars.Add(TEXT("ConsoleVariables")).CVarPriority = (int)ECVF_SetByPluginLowPriority;
-				ChangeTracker->CVars.Add(TEXT("ConsoleVariables_HighPriority")).CVarPriority = (int)ECVF_SetByPluginHighPriority;
+				ChangeTracker->CVars.Add(TEXT("PluginCVars_LowPriority")).CVarPriority = (int)ECVF_SetByPluginLowPriority;
+				ChangeTracker->CVars.Add(TEXT("PluginCVars_HighPriority")).CVarPriority = (int)ECVF_SetByPluginHighPriority;
 
 				// apply plugin modifications from this plugin to the everything
 				FConfigCacheIni::AddMultiplePluginsToAllBranches(PluginNames, ChangeTracker);
+
+#if !UE_BUILD_SHIPPING
+				for (const TPair<FName, TSet<FString>>& Pair : ChangeTracker->ModifiedSectionsPerBranch)
+				{
+					ensureAlwaysMsgf(!Pair.Value.Contains(TEXT("ConsoleVariables")) && !Pair.Value.Contains(TEXT("ConsoleVariables_HighPriority")),
+						TEXT("GameFeaturePlugins .ini files should no longer use [ConsoleVariables] or [ConsoleVariables_HighPriority] sections. Instead, use [PluginCVars_LowPriority] and [PluginCVars_HighPriority]."));
+				}
+#endif
 			});
 
 		UE::DynamicConfig::PerformDynamicConfig("InitializeHierarchicalPluginIniFiles", [&PluginNames](FConfigModificationTracker* ChangeTracker)
 			{
 				{
 					// set which sections to track for cvars, with their priority
-					ChangeTracker->CVars.Add(TEXT("ConsoleVariables")).CVarPriority = (int)ECVF_SetByHotfix;
-					ChangeTracker->CVars.Add(TEXT("ConsoleVariables_HighPriority")).CVarPriority = (int)ECVF_SetByHotfix;
+					ChangeTracker->CVars.Add(TEXT("PluginCVars_LowPriority")).CVarPriority = (int)ECVF_SetByHotfix;
+					ChangeTracker->CVars.Add(TEXT("PluginCVars_HighPriority")).CVarPriority = (int)ECVF_SetByHotfix;
 
 					for (const FName PluginName : PluginNames)
 					{
diff --git a/Engine/Source/Runtime/Core/Private/HAL/ConsoleManager.cpp b/Engine/Source/Runtime/Core/Private/HAL/ConsoleManager.cpp
--- a/Engine/Source/Runtime/Core/Private/HAL/ConsoleManager.cpp
+++ b/Engine/Source/Runtime/Core/Private/HAL/ConsoleManager.cpp
@@ -591,6 +591,12 @@ bool IConsoleManager::VisitPlatformCVarsForEmulation(FName PlatformName, const F
 		{ TEXT("Startup"), ECVF_SetByConsoleVariablesIni },
 		{ TEXT("ConsoleVariables"), ECVF_SetBySystemSettingsIni },
 
+		// Preview active GameFeaturePlugins - 
+		// NOTE: THIS ONLY WORKS FOR ENGINE.INI GFP CONFIG FILES!
+		//       because the code below only looks in the Engine branch, the GFPs don't run their Dynamic Modification code on other platforms configs
+		{ TEXT("PluginCVars_LowPriority"), ECVF_SetByPluginLowPriority },
+		{ TEXT("PluginCVars_HighPriority"), ECVF_SetByPluginHighPriority },
+
 	};
 
 	TMap<FString, int32> CVarSetByMap;
diff --git a/Engine/Source/Runtime/Core/Private/Misc/ConfigCacheIni.cpp b/Engine/Source/Runtime/Core/Private/Misc/ConfigCacheIni.cpp
--- a/Engine/Source/Runtime/Core/Private/Misc/ConfigCacheIni.cpp
+++ b/Engine/Source/Runtime/Core/Private/Misc/ConfigCacheIni.cpp
@@ -7286,8 +7286,8 @@ class FIniExec : public FSelfRegisteringExec
 				UE::DynamicConfig::PerformDynamicConfig("TrackedDyn", [Branch, Filename](FConfigModificationTracker* ChangeTracker)
 					{
 						// set which sections to track for cvars, with their priority
-						ChangeTracker->CVars.Add(TEXT("ConsoleVariables")).CVarPriority = (int)ECVF_SetByPluginLowPriority;
-						ChangeTracker->CVars.Add(TEXT("ConsoleVariables_HighPriority")).CVarPriority = (int)ECVF_SetByPluginHighPriority;
+						ChangeTracker->CVars.Add(TEXT("PluginCVars_LowPriority")).CVarPriority = (int)ECVF_SetByPluginLowPriority;
+						ChangeTracker->CVars.Add(TEXT("PluginCVars_HighPriority")).CVarPriority = (int)ECVF_SetByPluginHighPriority;
 
 						FDynamicLayerInfo Info { Filename, "TrackedDyn", (uint16)DynamicLayerPriority::Unknown };
 						Branch->AddDynamicLayersToHierarchy({ Info }, ChangeTracker);
