- Made GFP Hotfixes use the Hotfix naming ([ConsoleVariables], not [PluginCVars_LowPriority]) since they are setting at the Hotfix priority, it makes sense to make them the same as other Hotfixes

diff --git a/Engine/Plugins/Runtime/GameFeatures/Source/GameFeatures/Private/GameFeatureData.cpp b/Engine/Plugins/Runtime/GameFeatures/Source/GameFeatures/Private/GameFeatureData.cpp
--- a/Engine/Plugins/Runtime/GameFeatures/Source/GameFeatures/Private/GameFeatureData.cpp
+++ b/Engine/Plugins/Runtime/GameFeatures/Source/GameFeatures/Private/GameFeatureData.cpp
@@ -253,13 +253,12 @@ void UGameFeatureData::InitializeHierarchicalPluginIniFiles(const FString& Plugi
 				}
 #endif
 			});
-
 		UE::DynamicConfig::PerformDynamicConfig(PluginName, [PluginName](FConfigModificationTracker* ChangeTracker)
 			{
 				{
-					// set which sections to track for cvars, with their priority
-					ChangeTracker->CVars.Add(TEXT("PluginCVars_LowPriority")).CVarPriority = (int)ECVF_SetByHotfix;
-					ChangeTracker->CVars.Add(TEXT("PluginCVars_HighPriority")).CVarPriority = (int)ECVF_SetByHotfix;
+					// Hotfixes still use [ConsoleVariables], not Plugin_CVarsLow/HighPriority - since these are set at Hotfix priority, 
+					// we use the HF style naming
+					ChangeTracker->CVars.Add(TEXT("ConsoleVariables")).CVarPriority = (int)ECVF_SetByHotfix;
 
 					// give hotfix a chance to run at that point
 #define HOTFIX_BRANCH(Ini) UE::DynamicConfig::HotfixPluginForBranch.Broadcast(PluginName, #Ini, ChangeTracker);
@@ -268,6 +267,7 @@ void UGameFeatureData::InitializeHierarchicalPluginIniFiles(const FString& Plugi
 				}
 			});
 
+
 		return;
 	}
 
@@ -682,9 +682,9 @@ void UGameFeatureData::InitializeHierarchicalPluginIniFiles(const TArrayView<FSt
 		UE::DynamicConfig::PerformDynamicConfig("InitializeHierarchicalPluginIniFiles", [&PluginNames](FConfigModificationTracker* ChangeTracker)
 			{
 				{
-					// set which sections to track for cvars, with their priority
-					ChangeTracker->CVars.Add(TEXT("PluginCVars_LowPriority")).CVarPriority = (int)ECVF_SetByHotfix;
-					ChangeTracker->CVars.Add(TEXT("PluginCVars_HighPriority")).CVarPriority = (int)ECVF_SetByHotfix;
+					// Hotfixes still use [ConsoleVariables], not Plugin_CVarsLow/HighPriority - since these are set at Hotfix priority, 
+					// we use the HF style naming
+					ChangeTracker->CVars.Add(TEXT("ConsoleVariables")).CVarPriority = (int)ECVF_SetByHotfix;
 
 					for (const FName PluginName : PluginNames)
 					{
