[GFP] Increase verbosity for InstallBundleManger requests to decrease logspam

#rnx
#tests local cooked build match

diff --git a/Engine/Plugins/Runtime/GameFeatures/Source/GameFeatures/Private/GameFeaturePluginStateMachine.cpp b/Engine/Plugins/Runtime/GameFeatures/Source/GameFeatures/Private/GameFeaturePluginStateMachine.cpp
--- a/Engine/Plugins/Runtime/GameFeatures/Source/GameFeatures/Private/GameFeaturePluginStateMachine.cpp
+++ b/Engine/Plugins/Runtime/GameFeatures/Source/GameFeatures/Private/GameFeaturePluginStateMachine.cpp
@@ -17,6 +17,7 @@
 #include "Internationalization/PackageLocalizationManager.h"
 #include "Internationalization/StringTable.h"
 #include "IO/IoStoreOnDemand.h"
+#include "Logging/LogVerbosity.h"
 #include "Logging/StructuredLog.h"
 #include "Materials/MaterialInterface.h"
 #include "Misc/AsciiSet.h"
@@ -1361,7 +1362,7 @@ struct FBaseDataReleaseGameFeaturePluginState : public FGameFeaturePluginState
 		// Always set ExplicitRemoveList, GFPSharedInstallTracker has filtered out shared dependencies
 		EInstallBundleReleaseRequestFlags ReleaseFlags = GetReleaseRequestFlags() | EInstallBundleReleaseRequestFlags::ExplicitRemoveList;
 		TValueOrError<FInstallBundleReleaseRequestInfo, EInstallBundleResult> MaybeRequestInfo = BundleManager->RequestReleaseContent(
-			InstallBundlesToRelease, ReleaseFlags);
+			InstallBundlesToRelease, ReleaseFlags, TArrayView<const FName>(), ELogVerbosity::Verbose);
 
 		if (MaybeRequestInfo.HasError())
 		{
@@ -2197,7 +2198,7 @@ struct FGameFeaturePluginState_Unmounting : public FGameFeaturePluginState
 			EInstallBundleReleaseRequestFlags::SkipReleaseUnmountOnly |
 			EInstallBundleReleaseRequestFlags::ExplicitRemoveList; // Always set ExplicitRemoveList, GFPSharedInstallTracker has filtered out shared dependencies
 		TValueOrError<FInstallBundleReleaseRequestInfo, EInstallBundleResult> MaybeRequestInfo = BundleManager->RequestReleaseContent(
-			InstallBundlesToRelease, ReleaseFlags);
+			InstallBundlesToRelease, ReleaseFlags, TArrayView<const FName>(), ELogVerbosity::Verbose);
 
 		if (MaybeRequestInfo.HasError())
 		{
