If we add a plugin when reading as its not yet registered we will remove on exit of the scope

diff --git a/Engine/Plugins/Runtime/GameFeatures/Source/GameFeatures/Private/GameFeaturesSubsystem.cpp b/Engine/Plugins/Runtime/GameFeatures/Source/GameFeatures/Private/GameFeaturesSubsystem.cpp
--- a/Engine/Plugins/Runtime/GameFeatures/Source/GameFeatures/Private/GameFeaturesSubsystem.cpp
+++ b/Engine/Plugins/Runtime/GameFeatures/Source/GameFeatures/Private/GameFeaturesSubsystem.cpp
@@ -2887,15 +2887,32 @@ bool UGameFeaturesSubsystem::GetGameFeaturePluginDetailsInternal(const FString&
 		}
 	}
 
+
+#if WITH_EDITOR
+	// if we failed to find the plugin and add it to the plugin manager, we may leak this when closing a project
+	// so make sure to clean up the plugin when leaving this function
+	bool bPluginWasAdded = false;
+
+	ON_SCOPE_EXIT
+	{
+		if (bPluginWasAdded)
+		{
+			verify(IPluginManager::Get().RemoveFromPluginsList(PluginDescriptorFilename));
+		}
+	};
+#endif
+
 	TSharedPtr<FJsonObject> ObjectPtr;
 	{
 #if WITH_EDITOR
 		// In the editor we already have the plugin JSON cached
 		TSharedPtr<IPlugin> Plugin = IPluginManager::Get().FindPlugin(FPathViews::GetBaseFilename(PluginDescriptorFilename));
+
 		if (!Plugin.IsValid())
 		{
 			if (IPluginManager::Get().AddToPluginsList(PluginDescriptorFilename))
 			{
+				bPluginWasAdded = true;
 				Plugin = IPluginManager::Get().FindPlugin(FPathViews::GetBaseFilename(PluginDescriptorFilename));
 			}
 		}
