UE: Allow a pluginname to be provided for GetPreloadBundleStateForGameFeature. This way the policy can decide to not preload the asset bundles for primary assets in certain plugins.

diff --git a/Engine/Plugins/Runtime/GameFeatures/Source/GameFeatures/Private/GameFeaturePluginStateMachine.cpp b/Engine/Plugins/Runtime/GameFeatures/Source/GameFeatures/Private/GameFeaturePluginStateMachine.cpp
--- a/Engine/Plugins/Runtime/GameFeatures/Source/GameFeatures/Private/GameFeaturePluginStateMachine.cpp
+++ b/Engine/Plugins/Runtime/GameFeatures/Source/GameFeatures/Private/GameFeaturePluginStateMachine.cpp
@@ -3646,8 +3646,10 @@ struct FGameFeaturePluginState_Loading : public FGameFeaturePluginState
 
 				UpdateStateMachineDeferred();
 			});
-
-			BundleHandle = UAssetManager::Get().LoadPrimaryAssets(AssetIdsToLoad, Policy.GetPreloadBundleStateForGameFeature(), MoveTemp(LoadParams));
+			
+			FString PluginName;
+			GameFeatureToLoad->GetPluginName(PluginName);
+			BundleHandle = UAssetManager::Get().LoadPrimaryAssets(AssetIdsToLoad, Policy.GetPreloadBundleStateForPlugin(PluginName), MoveTemp(LoadParams));
 		}
 		else
 		{
diff --git a/Engine/Plugins/Runtime/GameFeatures/Source/GameFeatures/Private/GameFeaturesProjectPolicies.cpp b/Engine/Plugins/Runtime/GameFeatures/Source/GameFeatures/Private/GameFeaturesProjectPolicies.cpp
--- a/Engine/Plugins/Runtime/GameFeatures/Source/GameFeatures/Private/GameFeaturesProjectPolicies.cpp
+++ b/Engine/Plugins/Runtime/GameFeatures/Source/GameFeatures/Private/GameFeaturesProjectPolicies.cpp
@@ -46,7 +46,7 @@ void UDefaultGameFeaturesProjectPolicies::GetGameFeatureLoadingMode(bool& bLoadC
 	bLoadServerData = !IsRunningClientOnly();
 }
 
-const TArray<FName> UDefaultGameFeaturesProjectPolicies::GetPreloadBundleStateForGameFeature() const
+const TArray<FName> UDefaultGameFeaturesProjectPolicies::GetPreloadBundleStateForPlugin(const FString& PluginName) const
 {
 	// By default, use the bundles corresponding to loading mode
 	bool bLoadClientData, bLoadServerData;
diff --git a/Engine/Plugins/Runtime/GameFeatures/Source/GameFeatures/Public/GameFeaturesProjectPolicies.h b/Engine/Plugins/Runtime/GameFeatures/Source/GameFeatures/Public/GameFeaturesProjectPolicies.h
--- a/Engine/Plugins/Runtime/GameFeatures/Source/GameFeatures/Public/GameFeaturesProjectPolicies.h
+++ b/Engine/Plugins/Runtime/GameFeatures/Source/GameFeatures/Public/GameFeaturesProjectPolicies.h
@@ -49,7 +49,9 @@ public:
 
 	// Returns the bundle state to use for assets returned by GetPreloadAssetListForGameFeature()
 	// See the Asset Manager documentation for more information about asset bundles
+	// Eventually, we should deprecate GetPreloadBundleStateForGameFeature in favor of GetPreloadBundleStateForPlugin since having the plugin name is useful
 	virtual const TArray<FName> GetPreloadBundleStateForGameFeature() const { return TArray<FName>(); }
+	virtual const TArray<FName> GetPreloadBundleStateForPlugin(const FString& PluginName) const { return GetPreloadBundleStateForGameFeature(); }
 
 	// Called to determine if this should be treated as a client, server, or both for data preloading
 	// Actions can use this to decide what to load at runtime
@@ -129,7 +131,7 @@ public:
 	//~UGameFeaturesProjectPolicies interface
 	UE_API virtual void InitGameFeatureManager() override;
 	UE_API virtual void GetGameFeatureLoadingMode(bool& bLoadClientData, bool& bLoadServerData) const override;
-	UE_API virtual const TArray<FName> GetPreloadBundleStateForGameFeature() const override;
+	UE_API virtual const TArray<FName> GetPreloadBundleStateForPlugin(const FString& PluginName) const override;
 	//~End of UGameFeaturesProjectPolicies interface
 };
 
