Wrapping parts of the GFA with macro checks to prevent non-useable assets from loading in non-shipping builds
#rnx

diff --git a/Engine/Plugins/Runtime/GameFeatures/Source/GameFeatures/Private/GameFeatureAction_AddCheats.cpp b/Engine/Plugins/Runtime/GameFeatures/Source/GameFeatures/Private/GameFeatureAction_AddCheats.cpp
--- a/Engine/Plugins/Runtime/GameFeatures/Source/GameFeatures/Private/GameFeatureAction_AddCheats.cpp
+++ b/Engine/Plugins/Runtime/GameFeatures/Source/GameFeatures/Private/GameFeatureAction_AddCheats.cpp
@@ -16,12 +16,15 @@
 
 void UGameFeatureAction_AddCheats::OnGameFeatureActivating()
 {
+#if UE_WITH_CHEAT_MANAGER
 	bIsActive = true;
 	CheatManagerRegistrationHandle = UCheatManager::RegisterForOnCheatManagerCreated(FOnCheatManagerCreated::FDelegate::CreateUObject(this, &ThisClass::OnCheatManagerCreated));
+#endif // UE_WITH_CHEAT_MANAGER
 }
 
 void UGameFeatureAction_AddCheats::OnGameFeatureDeactivating(FGameFeatureDeactivatingContext& Context)
 {
+#if UE_WITH_CHEAT_MANAGER
 	UCheatManager::UnregisterFromOnCheatManagerCreated(CheatManagerRegistrationHandle);
 
 	for (TWeakObjectPtr<UCheatManagerExtension> ExtensionPtr : SpawnedCheatManagers)
@@ -34,6 +37,7 @@ void UGameFeatureAction_AddCheats::OnGameFeatureDeactivating(FGameFeatureDeactiv
 	}
 	SpawnedCheatManagers.Empty();
 	bIsActive = false;
+#endif // UE_WITH_CHEAT_MANAGER
 }
 
 #if WITH_EDITOR
@@ -58,6 +62,7 @@ EDataValidationResult UGameFeatureAction_AddCheats::IsDataValid(FDataValidationC
 
 void UGameFeatureAction_AddCheats::OnCheatManagerCreated(UCheatManager* CheatManager)
 {
+#if UE_WITH_CHEAT_MANAGER
 	// First clean out any stale pointers
 	for (int32 ManagerIdx = SpawnedCheatManagers.Num() - 1; ManagerIdx >= 0; --ManagerIdx)
 	{
@@ -111,16 +116,19 @@ void UGameFeatureAction_AddCheats::OnCheatManagerCreated(UCheatManager* CheatMan
 			}
 		}
 	}
+#endif // UE_WITH_CHEAT_MANAGER
 };
 
 void UGameFeatureAction_AddCheats::SpawnCheatManagerExtension(UCheatManager* CheatManager, const TSubclassOf<UCheatManagerExtension>& CheatManagerClass)
 {
+#if UE_WITH_CHEAT_MANAGER
 	if ((CheatManagerClass->ClassWithin == nullptr) || CheatManager->IsA(CheatManagerClass->ClassWithin))
 	{
 		UCheatManagerExtension* Extension = NewObject<UCheatManagerExtension>(CheatManager, CheatManagerClass);
 		SpawnedCheatManagers.Add(Extension);
 		CheatManager->AddCheatManagerExtension(Extension);
 	}
+#endif // UE_WITH_CHEAT_MANAGER
 }
 
 //////////////////////////////////////////////////////////////////////
