Once a GFP has been allowed to download through LoadBuiltInGameFeaturePlugin that GFP will be allowed to download in the future. This prevents multiple callers from requesting a plugin to be downloaded or not download unknown to the other and stomping another in flight download request by clobbering the GFP settings.

diff --git a/Engine/Plugins/Runtime/GameFeatures/Source/GameFeatures/Private/GameFeaturesSubsystem.cpp b/Engine/Plugins/Runtime/GameFeatures/Source/GameFeatures/Private/GameFeaturesSubsystem.cpp
--- a/Engine/Plugins/Runtime/GameFeatures/Source/GameFeatures/Private/GameFeaturesSubsystem.cpp
+++ b/Engine/Plugins/Runtime/GameFeatures/Source/GameFeatures/Private/GameFeaturesSubsystem.cpp
@@ -114,6 +114,13 @@ namespace UE::GameFeatures
 		TEXT("When call LoadBuiltInGameFeaturePlugins should only the plugins that were built by this target exe be considered built in."),
 		ECVF_Default);
 
+	static bool bAllowDownloadingAfterFirstDownloadRequest = true;
+	static FAutoConsoleVariableRef CVarAllowDownloadingAfterFirstRequest(
+		TEXT("GameFeaturePlugin.AllowDownloadingAfterFirstRequest"),
+		bAllowDownloadingAfterFirstDownloadRequest,
+		TEXT("Allow downloading a GFP when requested through LoadBuiltinGameFeaturePlugins after the first time it's set to be downloaded."),
+		ECVF_Default);
+
 	static FAutoConsoleVariable CVarTrackLoadStats(
 		TEXT("GameFeaturePlugin.TrackLoadStats"),
 		false,
@@ -2405,6 +2412,18 @@ void UGameFeaturesSubsystem::LoadBuiltInGameFeaturePlugin(const TSharedRef<IPlug
 		}
 		else
 		{
+			if (UE::GameFeatures::bAllowDownloadingAfterFirstDownloadRequest)
+			{
+				if (const UGameFeaturePluginStateMachine* GFSM = FindGameFeaturePluginStateMachine(PluginURL))
+				{
+					FGameFeatureProtocolOptions CurrentProtocolOptions = GFSM->GetProtocolOptions();
+					if (CurrentProtocolOptions.HasSubtype<FInstallBundlePluginProtocolOptions>())
+					{
+						const FInstallBundlePluginProtocolOptions& CurrentInstallBundleOptions = CurrentProtocolOptions.GetSubtype<FInstallBundlePluginProtocolOptions>();
+						BehaviorOptions.bDoNotDownload &= CurrentInstallBundleOptions.bDoNotDownload; // already allowed to be downloaded.
+					}
+				}
+			}
 			FGameFeatureProtocolOptions ProtocolOptions;
 			if (UGameFeaturesSubsystem::GetPluginURLProtocol(PluginURL) == EGameFeaturePluginProtocol::InstallBundle)
 			{
