Making UGameFeatureData::GetDisallowedClasses() virtual so game project subclasses can prevent some GF Action classes from being selectable.

diff --git a/Engine/Plugins/Runtime/GameFeatures/Source/GameFeatures/Public/GameFeatureData.h b/Engine/Plugins/Runtime/GameFeatures/Source/GameFeatures/Public/GameFeatureData.h
--- a/Engine/Plugins/Runtime/GameFeatures/Source/GameFeatures/Public/GameFeatureData.h
+++ b/Engine/Plugins/Runtime/GameFeatures/Source/GameFeatures/Public/GameFeatureData.h
@@ -115,15 +115,14 @@ public:
 #endif
 
 private:
-#if WITH_EDITOR
-	UFUNCTION()
-	UE_API TArray<UClass*> GetDisallowedActions() const;
-#endif
-
 	/** Internal helper function to reload config data on objects as a result of a plugin INI being loaded */
 	static UE_API void ReloadConfigs(FConfigFile& PluginConfig);
 
 protected:
+#if WITH_EDITOR
+	UFUNCTION()
+	UE_API virtual TArray<UClass*> GetDisallowedActions() const;
+#endif
 
 	/** List of actions to perform as this game feature is loaded/activated/deactivated/unloaded */
 	UPROPERTY(EditDefaultsOnly, Instanced, Category="Game Feature | Actions", meta = (GetDisallowedClasses = "GetDisallowedActions"))
