Add a mutex/scope lock around GameFeaturePluginNameToPathMap as its possible during some loading to try and access the GFP system
Caught during a TSan run

diff --git a/Engine/Plugins/Runtime/GameFeatures/Source/GameFeatures/Private/GameFeaturesSubsystem.cpp b/Engine/Plugins/Runtime/GameFeatures/Source/GameFeatures/Private/GameFeaturesSubsystem.cpp
--- a/Engine/Plugins/Runtime/GameFeatures/Source/GameFeatures/Private/GameFeaturesSubsystem.cpp
+++ b/Engine/Plugins/Runtime/GameFeatures/Source/GameFeatures/Private/GameFeaturesSubsystem.cpp
@@ -1209,6 +1209,8 @@ void UGameFeaturesSubsystem::OnGameFeatureTerminating(const FString& PluginName,
 
 	if (!PluginName.IsEmpty())
 	{
+		UE::TWriteScopeLock WriteLock(GameFeaturePluginNameToPathMapLock);
+
 		// Unmap plugin name to plugin URL
 		GameFeaturePluginNameToPathMap.Remove(TO_PLUGIN_STRING(PluginName));
 	}
@@ -1225,6 +1227,7 @@ void UGameFeaturesSubsystem::OnGameFeatureStatusKnown(const FString& PluginName,
 	FPluginString Name = TO_PLUGIN_STRING(PluginName);
 	PRAGMA_ENABLE_DEPRECATION_WARNINGS
 	// Map plugin name to plugin URL
+	UE::TWriteScopeLock WriteLock(GameFeaturePluginNameToPathMapLock);
 	if (ensure(!GameFeaturePluginNameToPathMap.Contains(Name)))
 	{
 		GameFeaturePluginNameToPathMap.Add(MoveTemp(Name), TO_PLUGIN_STRING(PluginIdentifier.GetFullPluginURL()));
@@ -2716,6 +2719,8 @@ void UGameFeaturesSubsystem::LoadBuiltInGameFeaturePluginsInternal(FBuiltInPlugi
 bool UGameFeaturesSubsystem::GetPluginURLByName(FStringView PluginName, FString& OutPluginURL) const
 {
 	PRAGMA_DISABLE_DEPRECATION_WARNINGS
+	UE::TReadScopeLock ReadLock(GameFeaturePluginNameToPathMapLock);
+
 	if (const FPluginString* PluginURL = GameFeaturePluginNameToPathMap.FindByHash(GetTypeHash(PluginName), TO_PLUGIN_STRING(PluginName)))	//TODO: GetTypeHash for utf8 strings only that doesn't match Tchar strings
 	{
 		PRAGMA_ENABLE_DEPRECATION_WARNINGS
@@ -2729,6 +2734,7 @@ bool UGameFeaturesSubsystem::GetPluginURLByName(FStringView PluginName, FString&
 bool UGameFeaturesSubsystem::GetPluginURLByName(FUtf8StringView PluginName, FString& OutPluginURL) const
 {
 	PRAGMA_DISABLE_DEPRECATION_WARNINGS
+	UE::TReadScopeLock ReadLock(GameFeaturePluginNameToPathMapLock);
 	if (const FPluginString* PluginURL = GameFeaturePluginNameToPathMap.FindByHash(GetTypeHash(PluginName), FPluginString(PluginName)))
 	{
 		PRAGMA_ENABLE_DEPRECATION_WARNINGS
diff --git a/Engine/Plugins/Runtime/GameFeatures/Source/GameFeatures/Public/GameFeaturesSubsystem.h b/Engine/Plugins/Runtime/GameFeatures/Source/GameFeatures/Public/GameFeaturesSubsystem.h
--- a/Engine/Plugins/Runtime/GameFeatures/Source/GameFeatures/Public/GameFeaturesSubsystem.h
+++ b/Engine/Plugins/Runtime/GameFeatures/Source/GameFeatures/Public/GameFeaturesSubsystem.h
@@ -961,6 +961,7 @@ private:
 
 	PRAGMA_DISABLE_DEPRECATION_WARNINGS
 	TMap<FPluginString, FPluginString> GameFeaturePluginNameToPathMap;
+	mutable FTransactionallySafeRWLock GameFeaturePluginNameToPathMapLock;
 	PRAGMA_ENABLE_DEPRECATION_WARNINGS
 
 	struct FCachedGameFeaturePluginDetails
