- When writing out the config file in FConfigFile::UpdateSections, don't write out the file if it matches the current file. This will make it so file like DefaultEngine.ini don't get marked as "changed" by some source control
- Slightly changed the COnfigModificationTracker class to be able to use bTrackLoadedFiles more often with worrying about tripping the code to not load a file that was already loaded (reducing side effects of using bTrackLoadedFiles).
- Now added context in the [ConsoleVariables]/[COnsoleVariables_HighPriority] deprecation ensures

diff --git a/Engine/Plugins/Runtime/GameFeatures/Source/GameFeatures/Private/GameFeatureData.cpp b/Engine/Plugins/Runtime/GameFeatures/Source/GameFeatures/Private/GameFeatureData.cpp
--- a/Engine/Plugins/Runtime/GameFeatures/Source/GameFeatures/Private/GameFeatureData.cpp
+++ b/Engine/Plugins/Runtime/GameFeatures/Source/GameFeatures/Private/GameFeatureData.cpp
@@ -285,6 +285,11 @@ void UGameFeatureData::InitializeHierarchicalPluginIniFiles(const FString& Plugi
 				ChangeTracker->CVars.Add(TEXT("PluginCVars_LowPriority")).CVarPriority = (int)ECVF_SetByPluginLowPriority;
 				ChangeTracker->CVars.Add(TEXT("PluginCVars_HighPriority")).CVarPriority = (int)ECVF_SetByPluginHighPriority;
 
+#if !UE_BUILD_SHIPPING
+				// track the loaded files so we can log potential filenames below in the ensure
+				ChangeTracker->bTrackLoadedFiles = true;
+#endif
+
 				// apply plugin modifications from this plugin to the everything
 				FConfigCacheIni::AddPluginToAllBranches(PluginName, ChangeTracker);
 
@@ -292,7 +297,8 @@ void UGameFeatureData::InitializeHierarchicalPluginIniFiles(const FString& Plugi
 				for (const TPair<FName, TSet<FString>>& Pair : ChangeTracker->ModifiedSectionsPerBranch)
 				{
 					ensureAlwaysMsgf(!Pair.Value.Contains(TEXT("ConsoleVariables")) && !Pair.Value.Contains(TEXT("ConsoleVariables_HighPriority")),
-						TEXT("GameFeaturePlugins .ini files should no longer use [ConsoleVariables] or [ConsoleVariables_HighPriority] sections. Instead, use [PluginCVars_LowPriority] and [PluginCVars_HighPriority]."));
+						TEXT("GameFeaturePlugins .ini files should no longer use [ConsoleVariables] or [ConsoleVariables_HighPriority] sections. Instead, use [PluginCVars_LowPriority] and [PluginCVars_HighPriority]. The old naming is in one of more of: %s"),
+						*FString::Join(ChangeTracker->LoadedFiles, TEXT(", ")));
 				}
 #endif
 			});
@@ -710,6 +716,11 @@ void UGameFeatureData::InitializeHierarchicalPluginIniFiles(const TArrayView<FSt
 				ChangeTracker->CVars.Add(TEXT("PluginCVars_LowPriority")).CVarPriority = (int)ECVF_SetByPluginLowPriority;
 				ChangeTracker->CVars.Add(TEXT("PluginCVars_HighPriority")).CVarPriority = (int)ECVF_SetByPluginHighPriority;
 
+#if !UE_BUILD_SHIPPING
+				// track the loaded files so we can log potential filenames below in the ensure
+				ChangeTracker->bTrackLoadedFiles = true;
+#endif
+
 				// apply plugin modifications from this plugin to the everything
 				FConfigCacheIni::AddMultiplePluginsToAllBranches(PluginNames, ChangeTracker);
 
@@ -717,7 +728,8 @@ void UGameFeatureData::InitializeHierarchicalPluginIniFiles(const TArrayView<FSt
 				for (const TPair<FName, TSet<FString>>& Pair : ChangeTracker->ModifiedSectionsPerBranch)
 				{
 					ensureAlwaysMsgf(!Pair.Value.Contains(TEXT("ConsoleVariables")) && !Pair.Value.Contains(TEXT("ConsoleVariables_HighPriority")),
-						TEXT("GameFeaturePlugins .ini files should no longer use [ConsoleVariables] or [ConsoleVariables_HighPriority] sections. Instead, use [PluginCVars_LowPriority] and [PluginCVars_HighPriority]."));
+						TEXT("GameFeaturePlugins .ini files should no longer use [ConsoleVariables] or [ConsoleVariables_HighPriority] sections. Instead, use [PluginCVars_LowPriority] and [PluginCVars_HighPriority]. The old naming is in one of more of: %s"),
+						*FString::Join(ChangeTracker->LoadedFiles, TEXT(", ")));
 				}
 #endif
 			});
diff --git a/Engine/Source/Runtime/Core/Private/Misc/ConfigCacheIni.cpp b/Engine/Source/Runtime/Core/Private/Misc/ConfigCacheIni.cpp
--- a/Engine/Source/Runtime/Core/Private/Misc/ConfigCacheIni.cpp
+++ b/Engine/Source/Runtime/Core/Private/Misc/ConfigCacheIni.cpp
@@ -2541,7 +2541,7 @@ bool FConfigFile::Write(const FString& Filename, bool bDoRemoteWrite/* = true*/,
 	{
 		SectionTexts.Add(FString(), PrefixText);
 	}
-	return WriteInternal(Filename, bDoRemoteWrite, SectionTexts, SectionOrder);
+	return WriteInternal(Filename, bDoRemoteWrite, SectionTexts, SectionOrder, nullptr);
 }
 
 void FConfigFile::WriteToString(FString& InOutText, const FString& SimulatedFilename /*= FString()*/, const FString& PrefixText /*= FString()*/)
@@ -2581,7 +2581,7 @@ bool FConfigFile::IsADefaultIniWrite(const FString& Filename, int32& OutIniCombi
 	return bIsADefaultIniWrite;
 }
 
-bool FConfigFile::WriteInternal(const FString& Filename, bool bDoRemoteWrite, TMap<FString, FString>& InOutSectionTexts, const TArray<FString>& InSectionOrder)
+bool FConfigFile::WriteInternal(const FString& Filename, bool bDoRemoteWrite, TMap<FString, FString>& InOutSectionTexts, const TArray<FString>& InSectionOrder, FString* OriginalContents)
 {
 	int32 IniCombineThreshold = MAX_int32;
 	bool bIsADefaultIniWrite = IsADefaultIniWrite(Filename, IniCombineThreshold);
@@ -2607,6 +2607,12 @@ bool FConfigFile::WriteInternal(const FString& Filename, bool bDoRemoteWrite, TM
 		FRemoteConfig::Get()->Write(*Filename, Text);
 	}
 
+	// don't write anything out if it didn't actually change
+	if (OriginalContents != nullptr && Text == *OriginalContents)
+	{
+		return true;
+	}
+
 	bool bResult = SaveConfigFileWrapper(*Filename, Text);
 
 	// File is still dirty if it didn't save.
@@ -6819,7 +6825,8 @@ void FConfigFile::UpdateSections(const TCHAR* DiskFilename, const TCHAR* IniRoot
 		Branch->IniPath = TEXT("");
 	}
 
-	WriteInternal(DiskFilename, true, SectionTexts, SectionOrder);
+	// pass in DiskFile, so it will not touch the disk if nothing changes
+	WriteInternal(DiskFilename, true, SectionTexts, SectionOrder, &DiskFile);
 }
 
 
@@ -6971,7 +6978,7 @@ void FConfigCacheIni::AddPluginsToBranches(const TArray<FName>& PluginNames, FCo
 			LoadedBranches.Add(BranchName);
 			
 			// if we have been tracking loaded files, and we've already loaded this file, we can skip it (it would be the DefaultMyPlugin.ini type of file)
-			if (ModificationTracker && ModificationTracker->bTrackLoadedFiles)
+			if (ModificationTracker && ModificationTracker->bSkipLoadingAlreadyLoadedFiles)
 			{
 				FString FullPath = bNamesAreFullPaths ? ConfigFilename : FPaths::Combine(PluginConfigDir, ConfigFilename);
 				if (ModificationTracker->LoadedFiles.Contains(FullPath))
diff --git a/Engine/Source/Runtime/Core/Public/Misc/ConfigCacheIni.h b/Engine/Source/Runtime/Core/Public/Misc/ConfigCacheIni.h
--- a/Engine/Source/Runtime/Core/Public/Misc/ConfigCacheIni.h
+++ b/Engine/Source/Runtime/Core/Public/Misc/ConfigCacheIni.h
@@ -509,6 +509,8 @@ public:
 	// input
 	bool bTrackModifiedSections = true;
 	bool bTrackLoadedFiles = false;
+	// if this is true and bTrackLoadedFiles is set, then we ignore files in LoadedFiles array
+	bool bSkipLoadingAlreadyLoadedFiles = false;
 
 	// output
 	TMap<FName, TSet<FString>> ModifiedSectionsPerBranch;
@@ -755,9 +757,10 @@ private:
 	 *  Any section in this array that is not found in InOutSectionTexts will be ignored.
 	 *  Any section in InOutSectionTexts that is not in this array will be appended to the end.
 	 *  Duplicate entries are ignored; the first found index is used.
+	 * @param OriginalContents - [optional] If set, and the string matches the generated string to write out, we skip the write to disk (it will return true)
 	 * @return TRUE if the write was successful
 	 */
-	bool WriteInternal(const FString& Filename, bool bDoRemoteWrite, TMap<FString, FString>& InOutSectionTexts, const TArray<FString>& InSectionOrder);
+	bool WriteInternal(const FString& Filename, bool bDoRemoteWrite, TMap<FString, FString>& InOutSectionTexts, const TArray<FString>& InSectionOrder, FString* OriginalContents);
 
 	/** Write a ConfigFile to InOutText, constructed from the given SectionTexts, in the given order, with sections in *this overriding sections in SectionTexts
 	 * @param InOutText - The string to write to
diff --git a/Engine/Source/Runtime/Projects/Private/PluginManager.cpp b/Engine/Source/Runtime/Projects/Private/PluginManager.cpp
--- a/Engine/Source/Runtime/Projects/Private/PluginManager.cpp
+++ b/Engine/Source/Runtime/Projects/Private/PluginManager.cpp
@@ -1901,6 +1901,7 @@ void FPluginManager::ProcessEnabledPlugins(FPluginDiscoveryContext& Context)
 	// Mount all the enabled plugins
 	FConfigModificationTracker ChangeTracker;
 	ChangeTracker.bTrackLoadedFiles = true;
+	ChangeTracker.bSkipLoadingAlreadyLoadedFiles = true;
 
 	TArray<FName> PluginNames;
 	PluginNames.Reserve(PluginsArray.Num());
