Add new batch Deactive, and Unload for GFPs so we can handle these faster
Move the suspend context out of Terminate/Release so we can share that logic with all 4 places

diff --git a/Engine/Plugins/Runtime/GameFeatures/Source/GameFeatures/Private/GameFeaturesSubsystem.cpp b/Engine/Plugins/Runtime/GameFeatures/Source/GameFeatures/Private/GameFeaturesSubsystem.cpp
--- a/Engine/Plugins/Runtime/GameFeatures/Source/GameFeatures/Private/GameFeaturesSubsystem.cpp
+++ b/Engine/Plugins/Runtime/GameFeatures/Source/GameFeatures/Private/GameFeaturesSubsystem.cpp
@@ -2031,6 +2031,100 @@ bool UGameFeaturesSubsystem::IsGameFeaturePluginActiveByName(FStringView PluginN
 	return true;
 }
 
+// Internal structs for suspending systems so we can better batch change GFPs from states with out spamming these systems
+namespace GameFeaturesSubsystem
+{
+	struct FSuspendPluginManger
+	{
+		UE_NODISCARD_CTOR FSuspendPluginManger()
+		{
+			IPluginManager::Get().SuppressPluginUnloadGC();
+		}
+
+		~FSuspendPluginManger()
+		{
+			IPluginManager::Get().ResumePluginUnloadGC();
+		}
+	};
+
+	struct FSuspendGameplayTagTreeChangedBroadcast
+	{
+		UE_NODISCARD_CTOR FSuspendGameplayTagTreeChangedBroadcast(bool bAllocateIfRequested)
+		{
+			if (UGameplayTagsManager* TagsManager = bAllocateIfRequested ? &UGameplayTagsManager::Get() : UGameplayTagsManager::GetIfAllocated() )
+			{
+				TagsManager->PushDeferOnGameplayTagTreeChangedBroadcast();
+				bWasSuspended = true;
+			}
+		}
+
+		~FSuspendGameplayTagTreeChangedBroadcast()
+		{
+			if (bWasSuspended)
+			{
+				UGameplayTagsManager::Get().PopDeferOnGameplayTagTreeChangedBroadcast();
+			}
+		}
+
+		bool bWasSuspended = false;
+	};
+
+	enum class ETransitionOptions
+	{
+		None							= 0,
+		SuspendPluginManager			= 1 << 0,
+		SuspendGameplayTagTreeBroadcast	= 1 << 1,
+		SuspendAssetRegistry			= 1 << 2,
+		All = SuspendPluginManager | SuspendGameplayTagTreeBroadcast | SuspendAssetRegistry,
+	};
+	ENUM_CLASS_FLAGS(ETransitionOptions)
+
+	struct FTransitionContenxt
+	{
+		TMap<FString, UE::GameFeatures::FResult> Results;
+		FMultipleGameFeaturePluginsReleased CompleteDelegate;
+		int32 NumPluginsTransitioned = 0;
+
+		TOptional<FSuspendPluginManger> SuspendPluginManager;
+		TOptional<FSuspendGameplayTagTreeChangedBroadcast> SuspendGameplayTagTreeBroadcast;
+		TOptional<IAssetRegistry::FSuspendPathsDismounting> SuspendAssetRegistry;
+
+		UE_NODISCARD_CTOR FTransitionContenxt(ETransitionOptions TransitionOptions = ETransitionOptions::All)
+		{
+			if (EnumHasAnyFlags(TransitionOptions, ETransitionOptions::SuspendPluginManager))
+			{
+				if (!IsEngineExitRequested())
+				{
+					SuspendPluginManager.Emplace();
+				}
+			}
+
+			if (EnumHasAnyFlags(TransitionOptions, ETransitionOptions::SuspendGameplayTagTreeBroadcast))
+			{
+				SuspendGameplayTagTreeBroadcast.Emplace(!IsEngineExitRequested());
+			}
+
+			if (EnumHasAnyFlags(TransitionOptions, ETransitionOptions::SuspendAssetRegistry))
+			{
+				SuspendAssetRegistry.Emplace();
+			}
+		}
+
+		~FTransitionContenxt()
+		{
+			CompleteDelegate.ExecuteIfBound(Results);
+			// Explicit order of destruction is required
+			// At the very least we need to execute PluginManger batch last
+			// as it will run GC and does the leaked packages verification.
+			// And by that point all other resources need to be properly cleared,
+			// i.e. other batches need to finish.
+			SuspendAssetRegistry.Reset();
+			SuspendGameplayTagTreeBroadcast.Reset();
+			SuspendPluginManager.Reset();
+		}
+	};
+}
+
 void UGameFeaturesSubsystem::DeactivateGameFeaturePlugin(const FString& PluginURL)
 {
 	DeactivateGameFeaturePlugin(PluginURL, FGameFeaturePluginDeactivateComplete());
@@ -2048,6 +2142,32 @@ void UGameFeaturesSubsystem::DeactivateGameFeaturePlugin(const FString& PluginUR
 	}
 }
 
+void UGameFeaturesSubsystem::DeactivateGameFeaturePlugin(TConstArrayView<FString> PluginURLs, const FMultipleGameFeaturePluginsDeactivate& CompleteDelegate)
+{
+	// None for now, as the current list are more of an issue going from Registered -> <lower>
+	GameFeaturesSubsystem::ETransitionOptions TransitionOptions = GameFeaturesSubsystem::ETransitionOptions::None;
+	TSharedRef<GameFeaturesSubsystem::FTransitionContenxt> TransitionContext = MakeShared<GameFeaturesSubsystem::FTransitionContenxt>(TransitionOptions);
+	TransitionContext->CompleteDelegate = CompleteDelegate;
+
+	TransitionContext->Results.Reserve(PluginURLs.Num());
+	for (const FString& PluginURL : PluginURLs)
+	{
+		TransitionContext->Results.Add(PluginURL, MakeError("Pending"));
+	}
+
+	UE_LOG(LogGameFeatures, Log, TEXT("Deactivating %i GFPs"), PluginURLs.Num());
+
+	for (const FString& PluginURL : PluginURLs)
+	{
+		DeactivateGameFeaturePlugin(PluginURL, FGameFeaturePluginDeactivateComplete::CreateLambda([TransitionContext, PluginURL](const UE::GameFeatures::FResult& Result)
+		{
+			TransitionContext->Results[PluginURL] = Result;
+			++TransitionContext->NumPluginsTransitioned;
+			UE_LOG(LogGameFeatures, VeryVerbose, TEXT("Finished Deactivating %i/%i GFPs"), TransitionContext->NumPluginsTransitioned, TransitionContext->Results.Num());
+		}));
+	}
+}
+
 void UGameFeaturesSubsystem::UnloadGameFeaturePlugin(const FString& PluginURL, bool bKeepRegistered /*= false*/)
 {
 	UnloadGameFeaturePlugin(PluginURL, FGameFeaturePluginUnloadComplete(), bKeepRegistered);
@@ -2066,6 +2186,37 @@ void UGameFeaturesSubsystem::UnloadGameFeaturePlugin(const FString& PluginURL, c
 	}
 }
 
+void UGameFeaturesSubsystem::UnloadGameFeaturePlugin(TConstArrayView<FString> PluginURLs, const FMultipleGameFeaturePluginsUnload& CompleteDelegate, bool bKeepRegistered /*= false*/)
+{
+	// Only set the downgrade options to suspend things when going below Registered for now
+	GameFeaturesSubsystem::ETransitionOptions TransitionOptions = GameFeaturesSubsystem::ETransitionOptions::None;
+	if (bKeepRegistered)
+	{
+		TransitionOptions = GameFeaturesSubsystem::ETransitionOptions::All;
+	}
+
+	TSharedRef<GameFeaturesSubsystem::FTransitionContenxt> TransitionContext = MakeShared<GameFeaturesSubsystem::FTransitionContenxt>(TransitionOptions);
+	TransitionContext->CompleteDelegate = CompleteDelegate;
+
+	TransitionContext->Results.Reserve(PluginURLs.Num());
+	for (const FString& PluginURL : PluginURLs)
+	{
+		TransitionContext->Results.Add(PluginURL, MakeError("Pending"));
+	}
+
+	UE_LOG(LogGameFeatures, Log, TEXT("Unloading %i GFPs bKeepRegistered(%i)"), PluginURLs.Num(), bKeepRegistered);
+
+	for (const FString& PluginURL : PluginURLs)
+	{
+		UnloadGameFeaturePlugin(PluginURL, FGameFeaturePluginUnloadComplete::CreateLambda([TransitionContext, PluginURL](const UE::GameFeatures::FResult& Result)
+		{
+			TransitionContext->Results[PluginURL] = Result;
+			++TransitionContext->NumPluginsTransitioned;
+			UE_LOG(LogGameFeatures, VeryVerbose, TEXT("Finished Unloading %i/%i GFPs"), TransitionContext->NumPluginsTransitioned, TransitionContext->Results.Num());
+		}), bKeepRegistered);
+	}
+}
+
 void UGameFeaturesSubsystem::ReleaseGameFeaturePlugin(const FString& PluginURL)
 {
 	ReleaseGameFeaturePlugin(PluginURL, FGameFeaturePluginReleaseComplete());
@@ -2085,80 +2236,24 @@ void UGameFeaturesSubsystem::ReleaseGameFeaturePlugin(const FString& PluginURL,
 
 void UGameFeaturesSubsystem::ReleaseGameFeaturePlugin(TConstArrayView<FString> PluginURLs, const FMultipleGameFeaturePluginsReleased& CompleteDelegate)
 {
-	struct FSuspendPluginManger
-	{
-		UE_NODISCARD_CTOR FSuspendPluginManger() { IPluginManager::Get().SuppressPluginUnloadGC(); }
-		~FSuspendPluginManger() { IPluginManager::Get().ResumePluginUnloadGC(); }
-	};
-	struct FSuspendGameplayTagTreeChangedBroadcast
-	{
-		UE_NODISCARD_CTOR FSuspendGameplayTagTreeChangedBroadcast(bool bAllocateIfRequested)
-		{
-			if (UGameplayTagsManager* TagsManager = bAllocateIfRequested ? &UGameplayTagsManager::Get() : UGameplayTagsManager::GetIfAllocated() )
-			{
-				TagsManager->PushDeferOnGameplayTagTreeChangedBroadcast();
-				bWasSuspended = true;
-			}
-		}
-		~FSuspendGameplayTagTreeChangedBroadcast()
-		{
-			if (bWasSuspended)
-			{
-				UGameplayTagsManager::Get().PopDeferOnGameplayTagTreeChangedBroadcast();
-			}
-		}
-
-		bool bWasSuspended = false;
-	};
-	struct FReleaseContext
-	{
-		TMap<FString, UE::GameFeatures::FResult> Results;
-		FMultipleGameFeaturePluginsReleased CompleteDelegate;
-		int32 NumPluginsReleased = 0;
+	TSharedRef<GameFeaturesSubsystem::FTransitionContenxt> TransitionContext = MakeShared<GameFeaturesSubsystem::FTransitionContenxt>();
+	TransitionContext->CompleteDelegate = CompleteDelegate;
 
-		TOptional<FSuspendPluginManger> SuspendPluginManager;
-		TOptional<FSuspendGameplayTagTreeChangedBroadcast> SuspendGameplayTagTreeBroadcast;
-		TOptional<IAssetRegistry::FSuspendPathsDismounting> SuspendAssetRegistry;
-
-		UE_NODISCARD_CTOR FReleaseContext()
-		{
-			if (!IsEngineExitRequested())
-			{
-				SuspendPluginManager.Emplace();
-			}
-			SuspendGameplayTagTreeBroadcast.Emplace(!IsEngineExitRequested());
-			SuspendAssetRegistry.Emplace();
-		}
-
-		~FReleaseContext()
-		{
-			CompleteDelegate.ExecuteIfBound(Results);
-			// Explicit order of destruction is required
-			// At the very least we need to execute PluginManger batch last
-			// as it will run GC and does the leaked packages verification.
-			// And by that point all other resources need to be properly cleared,
-			// i.e. other batches need to finish.
-			SuspendAssetRegistry.Reset();
-			SuspendGameplayTagTreeBroadcast.Reset();
-			SuspendPluginManager.Reset();
-		}
-	};
-	TSharedRef<FReleaseContext> LoadContext = MakeShared<FReleaseContext>();
-	LoadContext->CompleteDelegate = CompleteDelegate;
-
-	LoadContext->Results.Reserve(PluginURLs.Num());
+	TransitionContext->Results.Reserve(PluginURLs.Num());
 	for (const FString& PluginURL : PluginURLs)
 	{
-		LoadContext->Results.Add(PluginURL, MakeError("Pending"));
+		TransitionContext->Results.Add(PluginURL, MakeError("Pending"));
 	}
 
+	UE_LOG(LogGameFeatures, Log, TEXT("Releasing %i GFPs"), PluginURLs.Num());
+
 	for (const FString& PluginURL : PluginURLs)
 	{
-		ReleaseGameFeaturePlugin(PluginURL, FGameFeaturePluginReleaseComplete::CreateLambda([LoadContext, PluginURL](const UE::GameFeatures::FResult& Result)
+		ReleaseGameFeaturePlugin(PluginURL, FGameFeaturePluginReleaseComplete::CreateLambda([TransitionContext, PluginURL](const UE::GameFeatures::FResult& Result)
 		{
-			LoadContext->Results[PluginURL] = Result;
-			++LoadContext->NumPluginsReleased;
-			UE_LOG(LogGameFeatures, VeryVerbose, TEXT("Finished Releasing %i/%i GFPs"), LoadContext->NumPluginsReleased, LoadContext->Results.Num());
+			TransitionContext->Results[PluginURL] = Result;
+			++TransitionContext->NumPluginsTransitioned;
+			UE_LOG(LogGameFeatures, VeryVerbose, TEXT("Finished Releasing %i/%i GFPs"), TransitionContext->NumPluginsTransitioned, TransitionContext->Results.Num());
 		}));
 	}
 }
@@ -2251,60 +2346,24 @@ void UGameFeaturesSubsystem::TerminateGameFeaturePlugin(const FString& PluginURL
 
 void UGameFeaturesSubsystem::TerminateGameFeaturePlugin(TConstArrayView<FString> PluginURLs, const FMultipleGameFeaturePluginsTerminated& CompleteDelegate)
 {
-	struct FContext
-	{
-		TMap<FString, UE::GameFeatures::FResult> Results;
-		FMultipleGameFeaturePluginsTerminated CompleteDelegate;
-
-		int32 NumPluginsTerminated = 0;
-		bool bPushedTagsBroadcast = false;
-
-		FContext()
-		{
-			if (!IsEngineExitRequested())
-			{
-				UGameplayTagsManager::Get().PushDeferOnGameplayTagTreeChangedBroadcast();
-				bPushedTagsBroadcast = true;
-			}
-			else if (UGameplayTagsManager* TagsManager = UGameplayTagsManager::GetIfAllocated())
-			{
-				TagsManager->PushDeferOnGameplayTagTreeChangedBroadcast();
-				bPushedTagsBroadcast = true;
-			}
-		}
-
-		~FContext()
-		{
-			if (bPushedTagsBroadcast)
-			{
-				if (UGameplayTagsManager* TagsManager = UGameplayTagsManager::GetIfAllocated())
-				{
-					TagsManager->PopDeferOnGameplayTagTreeChangedBroadcast();
-				}
-			}
-
-			CompleteDelegate.ExecuteIfBound(Results);
-		}
-	};
-	TSharedRef<FContext> Context = MakeShared<FContext>();
-	Context->CompleteDelegate = CompleteDelegate;
+	TSharedRef<GameFeaturesSubsystem::FTransitionContenxt> TransitionContext = MakeShared<GameFeaturesSubsystem::FTransitionContenxt>();
+	TransitionContext->CompleteDelegate = CompleteDelegate;
 
-	Context->Results.Reserve(PluginURLs.Num());
+	TransitionContext->Results.Reserve(PluginURLs.Num());
 	for (const FString& PluginURL : PluginURLs)
 	{
-		Context->Results.Add(PluginURL, MakeError("Pending"));
+		TransitionContext->Results.Add(PluginURL, MakeError("Pending"));
 	}
 
-	const int32 NumPlugins = PluginURLs.Num();
-	UE_LOG(LogGameFeatures, Log, TEXT("Terminating %i GFPs"), NumPlugins);
+	UE_LOG(LogGameFeatures, Log, TEXT("Terminating %i GFPs"), PluginURLs.Num());
 
 	for (const FString& PluginURL : PluginURLs)
 	{
-		TerminateGameFeaturePlugin(PluginURL, FGameFeaturePluginTerminateComplete::CreateLambda([Context, PluginURL](const UE::GameFeatures::FResult& Result)
+		TerminateGameFeaturePlugin(PluginURL, FGameFeaturePluginTerminateComplete::CreateLambda([TransitionContext, PluginURL](const UE::GameFeatures::FResult& Result)
 		{
-			Context->Results.Add(PluginURL, Result);
-			++Context->NumPluginsTerminated;
-			UE_LOG(LogGameFeatures, VeryVerbose, TEXT("Finished Terminating %i GFPs"), Context->NumPluginsTerminated);
+			TransitionContext->Results.Add(PluginURL, Result);
+			++TransitionContext->NumPluginsTransitioned;
+			UE_LOG(LogGameFeatures, VeryVerbose, TEXT("Finished Terminating %i GFPs"), TransitionContext->NumPluginsTransitioned);
 		}));
 	}
 }
diff --git a/Engine/Plugins/Runtime/GameFeatures/Source/GameFeatures/Public/GameFeaturesSubsystem.h b/Engine/Plugins/Runtime/GameFeatures/Source/GameFeatures/Public/GameFeaturesSubsystem.h
--- a/Engine/Plugins/Runtime/GameFeatures/Source/GameFeatures/Public/GameFeaturesSubsystem.h
+++ b/Engine/Plugins/Runtime/GameFeatures/Source/GameFeatures/Public/GameFeaturesSubsystem.h
@@ -221,6 +221,8 @@ using FMultipleGameFeaturePluginChangeStateComplete = TDelegate<void(const TMap<
 
 using FBuiltInGameFeaturePluginsLoaded = FMultipleGameFeaturePluginChangeStateComplete;
 using FMultipleGameFeaturePluginsLoaded = FMultipleGameFeaturePluginChangeStateComplete;
+using FMultipleGameFeaturePluginsDeactivate = FMultipleGameFeaturePluginChangeStateComplete;
+using FMultipleGameFeaturePluginsUnload = FMultipleGameFeaturePluginChangeStateComplete;
 using FMultipleGameFeaturePluginsReleased = FMultipleGameFeaturePluginChangeStateComplete;
 using FMultipleGameFeaturePluginsTerminated = FMultipleGameFeaturePluginChangeStateComplete;
 
@@ -624,15 +626,17 @@ public:
 	/** Determines if a plugin is up to date or needs an update. Returns true if an update is available.*/
 	UE_API bool DoesGameFeaturePluginNeedUpdate(const FString& PluginURL) const;
 
-	/** Deactivates the specified plugin */
+	/** Deactivates the specified plugin. Transition to [Terminal, Loaded] */
 	UE_API void DeactivateGameFeaturePlugin(const FString& PluginURL);
 	UE_API void DeactivateGameFeaturePlugin(const FString& PluginURL, const FGameFeaturePluginDeactivateComplete& CompleteDelegate);
+	UE_API void DeactivateGameFeaturePlugin(TConstArrayView<FString> PluginURLs, const FMultipleGameFeaturePluginsDeactivate& CompleteDelegate);
 
-	/** Unloads the specified game feature plugin. */
+	/** Unloads the specified game feature plugin. Transition to bKeepRegistered = false -> [Terminal, Installed], bKeepRegistered = true -> [Terminal, Registered]*/
 	UE_API void UnloadGameFeaturePlugin(const FString& PluginURL, bool bKeepRegistered = false);
 	UE_API void UnloadGameFeaturePlugin(const FString& PluginURL, const FGameFeaturePluginUnloadComplete& CompleteDelegate, bool bKeepRegistered = false);
+	UE_API void UnloadGameFeaturePlugin(TConstArrayView<FString> PluginURLs, const FMultipleGameFeaturePluginsUnload& CompleteDelegate, bool bKeepRegistered = false);
 
-	/** Releases any game data stored for this GameFeaturePlugin. Does not uninstall data and it will remain on disk. */
+	/** Releases any game data stored for this GameFeaturePlugin. Does not uninstall data and it will remain on disk. [Terminal, StatusKnown] */
 	UE_API void ReleaseGameFeaturePlugin(const FString& PluginURL);
 	UE_API void ReleaseGameFeaturePlugin(const FString& PluginURL, const FGameFeaturePluginReleaseComplete& CompleteDelegate);
 	UE_API void ReleaseGameFeaturePlugin(TConstArrayView<FString> PluginURLs, const FMultipleGameFeaturePluginsReleased& CompleteDelegate);
@@ -644,11 +648,11 @@ public:
 	UE_API void UninstallGameFeaturePlugin(const FString& PluginURL, const FGameFeaturePluginUninstallComplete& CompleteDelegate = FGameFeaturePluginUninstallComplete());
 	UE_API void UninstallGameFeaturePlugin(const FString& PluginURL, const FGameFeatureProtocolOptions& ProtocolOptions, const FGameFeaturePluginUninstallComplete& CompleteDelegate = FGameFeaturePluginUninstallComplete());
 
-	/** Terminate the GameFeaturePlugin and remove all associated plugin tracking data. */
+	/** Terminate the GameFeaturePlugin and remove all associated plugin tracking data. Transition to [Terminal, Terminal] */
 	UE_API void TerminateGameFeaturePlugin(const FString& PluginURL);
 	UE_API void TerminateGameFeaturePlugin(const FString& PluginURL, const FGameFeaturePluginTerminateComplete& CompleteDelegate);
 	UE_API void TerminateGameFeaturePlugin(TConstArrayView<FString> PluginURLs, const FMultipleGameFeaturePluginsTerminated& CompleteDelegate);
-	
+
 	/** Attempt to cancel any state change. Calls back when cancelation is complete. Any other pending callbacks will be called with a canceled error. */
 	UE_API void CancelGameFeatureStateChange(const FString& PluginURL);
 	UE_API void CancelGameFeatureStateChange(const FString& PluginURL, const FGameFeaturePluginChangeStateComplete& CompleteDelegate);
