Clean up the load context to use the new TransitionContext so we can re-use the same code

diff --git a/Engine/Plugins/Runtime/GameFeatures/Source/GameFeatures/Private/GameFeaturesSubsystem.cpp b/Engine/Plugins/Runtime/GameFeatures/Source/GameFeatures/Private/GameFeaturesSubsystem.cpp
--- a/Engine/Plugins/Runtime/GameFeatures/Source/GameFeatures/Private/GameFeaturesSubsystem.cpp
+++ b/Engine/Plugins/Runtime/GameFeatures/Source/GameFeatures/Private/GameFeaturesSubsystem.cpp
@@ -353,6 +353,100 @@ void LexFromString(EGameFeatureTargetState& ValueOut, const TCHAR* StringIn)
 	ensureAlwaysMsgf(false, TEXT("Invalid FString { %s } used for EGameFeatureTargetState LexFromString. Does not correspond to any EGameFeatureTargetState!"), StringIn);
 }
 
+// Internal structs for suspending systems so we can better batch change GFPs from states with out spamming these systems
+namespace GameFeaturesSubsystem
+{
+	struct FSuspendPluginManger
+	{
+		UE_NODISCARD_CTOR FSuspendPluginManger()
+		{
+			IPluginManager::Get().SuppressPluginUnloadGC();
+		}
+
+		~FSuspendPluginManger()
+		{
+			IPluginManager::Get().ResumePluginUnloadGC();
+		}
+	};
+
+	struct FSuspendGameplayTagTreeChangedBroadcast
+	{
+		UE_NODISCARD_CTOR FSuspendGameplayTagTreeChangedBroadcast(bool bAllocateIfRequested)
+		{
+			if (UGameplayTagsManager* TagsManager = bAllocateIfRequested ? &UGameplayTagsManager::Get() : UGameplayTagsManager::GetIfAllocated() )
+			{
+				TagsManager->PushDeferOnGameplayTagTreeChangedBroadcast();
+				bWasSuspended = true;
+			}
+		}
+
+		~FSuspendGameplayTagTreeChangedBroadcast()
+		{
+			if (bWasSuspended)
+			{
+				UGameplayTagsManager::Get().PopDeferOnGameplayTagTreeChangedBroadcast();
+			}
+		}
+
+		bool bWasSuspended = false;
+	};
+
+	enum class ETransitionOptions
+	{
+		None							= 0,
+		SuspendPluginManager			= 1 << 0,
+		SuspendGameplayTagTreeBroadcast	= 1 << 1,
+		SuspendAssetRegistry			= 1 << 2,
+		All = SuspendPluginManager | SuspendGameplayTagTreeBroadcast | SuspendAssetRegistry,
+	};
+	ENUM_CLASS_FLAGS(ETransitionOptions)
+
+	struct FTransitionContenxt
+	{
+		TMap<FString, UE::GameFeatures::FResult> Results;
+		FMultipleGameFeaturePluginsReleased CompleteDelegate;
+		int32 NumPluginsTransitioned = 0;
+
+		TOptional<FSuspendPluginManger> SuspendPluginManager;
+		TOptional<FSuspendGameplayTagTreeChangedBroadcast> SuspendGameplayTagTreeBroadcast;
+		TOptional<IAssetRegistry::FSuspendPathsDismounting> SuspendAssetRegistry;
+
+		UE_NODISCARD_CTOR FTransitionContenxt(ETransitionOptions TransitionOptions = ETransitionOptions::All)
+		{
+			if (EnumHasAnyFlags(TransitionOptions, ETransitionOptions::SuspendPluginManager))
+			{
+				if (!IsEngineExitRequested())
+				{
+					SuspendPluginManager.Emplace();
+				}
+			}
+
+			if (EnumHasAnyFlags(TransitionOptions, ETransitionOptions::SuspendGameplayTagTreeBroadcast))
+			{
+				SuspendGameplayTagTreeBroadcast.Emplace(!IsEngineExitRequested());
+			}
+
+			if (EnumHasAnyFlags(TransitionOptions, ETransitionOptions::SuspendAssetRegistry))
+			{
+				SuspendAssetRegistry.Emplace();
+			}
+		}
+
+		~FTransitionContenxt()
+		{
+			CompleteDelegate.ExecuteIfBound(Results);
+			// Explicit order of destruction is required
+			// At the very least we need to execute PluginManger batch last
+			// as it will run GC and does the leaked packages verification.
+			// And by that point all other resources need to be properly cleared,
+			// i.e. other batches need to finish.
+			SuspendAssetRegistry.Reset();
+			SuspendGameplayTagTreeBroadcast.Reset();
+			SuspendPluginManager.Reset();
+		}
+	};
+}
+
 FGameFeaturePluginIdentifier::FGameFeaturePluginIdentifier(FString PluginURL)
 {
 	FromPluginURL(MoveTemp(PluginURL));
@@ -1528,60 +1622,25 @@ void UGameFeaturesSubsystem::LoadGameFeaturePlugin(const FString& PluginURL, con
 
 void UGameFeaturesSubsystem::LoadGameFeaturePlugin(TConstArrayView<FString> PluginURLs, const FGameFeatureProtocolOptions& ProtocolOptions, const FMultipleGameFeaturePluginsLoaded& CompleteDelegate)
 {
-	struct FLoadContext
-	{
-		TMap<FString, UE::GameFeatures::FResult> Results;
-		FMultipleGameFeaturePluginsLoaded CompleteDelegate;
-
-		int32 NumPluginsLoaded = 0;
-		bool bPushedTagsBroadcast = false;
-
-		FLoadContext()
-		{
-			if (!IsEngineExitRequested())
-			{
-				UGameplayTagsManager::Get().PushDeferOnGameplayTagTreeChangedBroadcast();
-				bPushedTagsBroadcast = true;
-			}
-			else if(UGameplayTagsManager* TagsManager = UGameplayTagsManager::GetIfAllocated())
-			{
-				TagsManager->PushDeferOnGameplayTagTreeChangedBroadcast();
-				bPushedTagsBroadcast = true;
-			}
-		}
-
-		~FLoadContext()
-		{
-			if (bPushedTagsBroadcast)
-			{
-				if (UGameplayTagsManager* TagsManager = UGameplayTagsManager::GetIfAllocated())
-				{
-					TagsManager->PopDeferOnGameplayTagTreeChangedBroadcast();
-				}
-			}
-
-			CompleteDelegate.ExecuteIfBound(Results);
-		}
-	};
-	TSharedRef<FLoadContext> LoadContext = MakeShared<FLoadContext>();
-	LoadContext->CompleteDelegate = CompleteDelegate;
+	GameFeaturesSubsystem::ETransitionOptions TransitionOptions = GameFeaturesSubsystem::ETransitionOptions::SuspendGameplayTagTreeBroadcast;
+	TSharedRef<GameFeaturesSubsystem::FTransitionContenxt> TransitionContext = MakeShared<GameFeaturesSubsystem::FTransitionContenxt>(TransitionOptions);
+	TransitionContext->CompleteDelegate = CompleteDelegate;
 
-	LoadContext->Results.Reserve(PluginURLs.Num());
+	TransitionContext->Results.Reserve(PluginURLs.Num());
 	for (const FString& PluginURL : PluginURLs)
 	{
-		LoadContext->Results.Add(PluginURL, MakeError("Pending"));
+		TransitionContext->Results.Add(PluginURL, MakeError("Pending"));
 	}
 
-	const int32 NumPluginsToLoad = PluginURLs.Num();
-	UE_LOG(LogGameFeatures, Log, TEXT("Loading %i GFPs"), NumPluginsToLoad);
+	UE_LOG(LogGameFeatures, Log, TEXT("Loading %i GFPs"), PluginURLs.Num());
 
 	for (const FString& PluginURL : PluginURLs)
 	{
-		LoadGameFeaturePlugin(PluginURL, ProtocolOptions, FGameFeaturePluginChangeStateComplete::CreateLambda([LoadContext, PluginURL](const UE::GameFeatures::FResult& Result)
+		LoadGameFeaturePlugin(PluginURL, ProtocolOptions, FGameFeaturePluginChangeStateComplete::CreateLambda([TransitionContext, PluginURL](const UE::GameFeatures::FResult& Result)
 		{
-			LoadContext->Results.Add(PluginURL, Result);
-			++LoadContext->NumPluginsLoaded;
-			UE_LOG(LogGameFeatures, VeryVerbose, TEXT("Finished Loading %i GFPs"), LoadContext->NumPluginsLoaded);
+			TransitionContext->Results.Add(PluginURL, Result);
+			++TransitionContext->NumPluginsTransitioned;
+			UE_LOG(LogGameFeatures, VeryVerbose, TEXT("Finished Loading %i GFPs"), TransitionContext->NumPluginsTransitioned);
 		}));
 	}
 }
@@ -1629,60 +1688,24 @@ void UGameFeaturesSubsystem::RegisterGameFeaturePlugin(const FString& PluginURL,
 
 void UGameFeaturesSubsystem::RegisterGameFeaturePlugin(TConstArrayView<FString> PluginURLs, const FGameFeatureProtocolOptions& ProtocolOptions, const FMultipleGameFeaturePluginsLoaded& CompleteDelegate)
 {
-	struct FLoadContext
-	{
-		TMap<FString, UE::GameFeatures::FResult> Results;
-		FMultipleGameFeaturePluginsLoaded CompleteDelegate;
-
-		int32 NumPluginsLoaded = 0;
-		bool bPushedTagsBroadcast = false;
-
-		FLoadContext()
-		{
-			if (!IsEngineExitRequested())
-			{
-				UGameplayTagsManager::Get().PushDeferOnGameplayTagTreeChangedBroadcast();
-				bPushedTagsBroadcast = true;
-			}
-			else if (UGameplayTagsManager* TagsManager = UGameplayTagsManager::GetIfAllocated())
-			{
-				TagsManager->PushDeferOnGameplayTagTreeChangedBroadcast();
-				bPushedTagsBroadcast = true;
-			}
-		}
-
-		~FLoadContext()
-		{
-			if (bPushedTagsBroadcast)
-			{
-				if (UGameplayTagsManager* TagsManager = UGameplayTagsManager::GetIfAllocated())
-				{
-					TagsManager->PopDeferOnGameplayTagTreeChangedBroadcast();
-				}
-			}
-
-			CompleteDelegate.ExecuteIfBound(Results);
-		}
-	};
-	TSharedRef<FLoadContext> LoadContext = MakeShared<FLoadContext>();
-	LoadContext->CompleteDelegate = CompleteDelegate;
+	GameFeaturesSubsystem::ETransitionOptions TransitionOptions = GameFeaturesSubsystem::ETransitionOptions::SuspendGameplayTagTreeBroadcast;
+	TSharedRef<GameFeaturesSubsystem::FTransitionContenxt> TransitionContext = MakeShared<GameFeaturesSubsystem::FTransitionContenxt>(TransitionOptions);
 
-	LoadContext->Results.Reserve(PluginURLs.Num());
+	TransitionContext->Results.Reserve(PluginURLs.Num());
 	for (const FString& PluginURL : PluginURLs)
 	{
-		LoadContext->Results.Add(PluginURL, MakeError("Pending"));
+		TransitionContext->Results.Add(PluginURL, MakeError("Pending"));
 	}
 
-	const int32 NumPluginsToLoad = PluginURLs.Num();
-	UE_LOG(LogGameFeatures, Log, TEXT("Registering %i GFPs"), NumPluginsToLoad);
+	UE_LOG(LogGameFeatures, Log, TEXT("Registering %i GFPs"), PluginURLs.Num());
 
 	for (const FString& PluginURL : PluginURLs)
 	{
-		RegisterGameFeaturePlugin(PluginURL, ProtocolOptions, FGameFeaturePluginChangeStateComplete::CreateLambda([LoadContext, PluginURL](const UE::GameFeatures::FResult& Result)
+		RegisterGameFeaturePlugin(PluginURL, ProtocolOptions, FGameFeaturePluginChangeStateComplete::CreateLambda([TransitionContext, PluginURL](const UE::GameFeatures::FResult& Result)
 			{
-				LoadContext->Results.Add(PluginURL, Result);
-				++LoadContext->NumPluginsLoaded;
-				UE_LOG(LogGameFeatures, VeryVerbose, TEXT("Finished Registering %i GFPs"), LoadContext->NumPluginsLoaded);
+				TransitionContext->Results.Add(PluginURL, Result);
+				++TransitionContext->NumPluginsTransitioned;
+				UE_LOG(LogGameFeatures, VeryVerbose, TEXT("Finished Registering %i GFPs"), TransitionContext->NumPluginsTransitioned);
 			}));
 	}
 }
@@ -1787,48 +1810,14 @@ void UGameFeaturesSubsystem::ChangeGameFeatureTargetState(const FString& PluginU
 
 void UGameFeaturesSubsystem::ChangeGameFeatureTargetState(TConstArrayView<FString> PluginURLs, const FGameFeatureProtocolOptions& ProtocolOptions, EGameFeatureTargetState TargetState, const FMultipleGameFeaturePluginsLoaded& CompleteDelegate)
 {
-	struct FLoadContext
-	{
-		TMap<FString, UE::GameFeatures::FResult> Results;
-		FMultipleGameFeaturePluginsLoaded CompleteDelegate;
-
-		int32 NumPluginsLoaded = 0;
-		bool bPushedTagsBroadcast = false;
-
-		FLoadContext()
-		{
-			if (!IsEngineExitRequested())
-			{
-				UGameplayTagsManager::Get().PushDeferOnGameplayTagTreeChangedBroadcast();
-				bPushedTagsBroadcast = true;
-			}
-			else if(UGameplayTagsManager* TagsManager = UGameplayTagsManager::GetIfAllocated())
-			{
-				TagsManager->PushDeferOnGameplayTagTreeChangedBroadcast();
-				bPushedTagsBroadcast = true;
-			}
-		}
-
-		~FLoadContext()
-		{
-			if (bPushedTagsBroadcast)
-			{
-				if (UGameplayTagsManager* TagsManager = UGameplayTagsManager::GetIfAllocated())
-				{
-					TagsManager->PopDeferOnGameplayTagTreeChangedBroadcast();
-				}
-			}
-
-			CompleteDelegate.ExecuteIfBound(Results);
-		}
-	};
-	TSharedRef<FLoadContext> LoadContext = MakeShared<FLoadContext>();
-	LoadContext->CompleteDelegate = CompleteDelegate;
+	GameFeaturesSubsystem::ETransitionOptions TransitionOptions = GameFeaturesSubsystem::ETransitionOptions::SuspendGameplayTagTreeBroadcast;
+	TSharedRef<GameFeaturesSubsystem::FTransitionContenxt> TransitionContext = MakeShared<GameFeaturesSubsystem::FTransitionContenxt>(TransitionOptions);
+	TransitionContext->CompleteDelegate = CompleteDelegate;
 
-	LoadContext->Results.Reserve(PluginURLs.Num());
+	TransitionContext->Results.Reserve(PluginURLs.Num());
 	for (const FString& PluginURL : PluginURLs)
 	{
-		LoadContext->Results.Add(PluginURL, MakeError("Pending"));
+		TransitionContext->Results.Add(PluginURL, MakeError("Pending"));
 	}
 
 	const int32 NumPluginsToLoad = PluginURLs.Num();
@@ -1836,11 +1825,11 @@ void UGameFeaturesSubsystem::ChangeGameFeatureTargetState(TConstArrayView<FStrin
 
 	for (const FString& PluginURL : PluginURLs)
 	{
-		ChangeGameFeatureTargetState(PluginURL, ProtocolOptions, TargetState, FGameFeaturePluginChangeStateComplete::CreateLambda([LoadContext, PluginURL, TargetState](const UE::GameFeatures::FResult& Result)
+		ChangeGameFeatureTargetState(PluginURL, ProtocolOptions, TargetState, FGameFeaturePluginChangeStateComplete::CreateLambda([TransitionContext, PluginURL, TargetState](const UE::GameFeatures::FResult& Result)
 		{
-			LoadContext->Results.Add(PluginURL, Result);
-			++LoadContext->NumPluginsLoaded;
-			UE_LOG(LogGameFeatures, VeryVerbose, TEXT("Finished Transitioning (%s) %i GFPs"), *LexToString(TargetState), LoadContext->NumPluginsLoaded);
+			TransitionContext->Results.Add(PluginURL, Result);
+			++TransitionContext->NumPluginsTransitioned;
+			UE_LOG(LogGameFeatures, VeryVerbose, TEXT("Finished Transitioning (%s) %i GFPs"), *LexToString(TargetState), TransitionContext->NumPluginsTransitioned);
 		}));
 	}
 }
@@ -2031,100 +2020,6 @@ bool UGameFeaturesSubsystem::IsGameFeaturePluginActiveByName(FStringView PluginN
 	return true;
 }
 
-// Internal structs for suspending systems so we can better batch change GFPs from states with out spamming these systems
-namespace GameFeaturesSubsystem
-{
-	struct FSuspendPluginManger
-	{
-		UE_NODISCARD_CTOR FSuspendPluginManger()
-		{
-			IPluginManager::Get().SuppressPluginUnloadGC();
-		}
-
-		~FSuspendPluginManger()
-		{
-			IPluginManager::Get().ResumePluginUnloadGC();
-		}
-	};
-
-	struct FSuspendGameplayTagTreeChangedBroadcast
-	{
-		UE_NODISCARD_CTOR FSuspendGameplayTagTreeChangedBroadcast(bool bAllocateIfRequested)
-		{
-			if (UGameplayTagsManager* TagsManager = bAllocateIfRequested ? &UGameplayTagsManager::Get() : UGameplayTagsManager::GetIfAllocated() )
-			{
-				TagsManager->PushDeferOnGameplayTagTreeChangedBroadcast();
-				bWasSuspended = true;
-			}
-		}
-
-		~FSuspendGameplayTagTreeChangedBroadcast()
-		{
-			if (bWasSuspended)
-			{
-				UGameplayTagsManager::Get().PopDeferOnGameplayTagTreeChangedBroadcast();
-			}
-		}
-
-		bool bWasSuspended = false;
-	};
-
-	enum class ETransitionOptions
-	{
-		None							= 0,
-		SuspendPluginManager			= 1 << 0,
-		SuspendGameplayTagTreeBroadcast	= 1 << 1,
-		SuspendAssetRegistry			= 1 << 2,
-		All = SuspendPluginManager | SuspendGameplayTagTreeBroadcast | SuspendAssetRegistry,
-	};
-	ENUM_CLASS_FLAGS(ETransitionOptions)
-
-	struct FTransitionContenxt
-	{
-		TMap<FString, UE::GameFeatures::FResult> Results;
-		FMultipleGameFeaturePluginsReleased CompleteDelegate;
-		int32 NumPluginsTransitioned = 0;
-
-		TOptional<FSuspendPluginManger> SuspendPluginManager;
-		TOptional<FSuspendGameplayTagTreeChangedBroadcast> SuspendGameplayTagTreeBroadcast;
-		TOptional<IAssetRegistry::FSuspendPathsDismounting> SuspendAssetRegistry;
-
-		UE_NODISCARD_CTOR FTransitionContenxt(ETransitionOptions TransitionOptions = ETransitionOptions::All)
-		{
-			if (EnumHasAnyFlags(TransitionOptions, ETransitionOptions::SuspendPluginManager))
-			{
-				if (!IsEngineExitRequested())
-				{
-					SuspendPluginManager.Emplace();
-				}
-			}
-
-			if (EnumHasAnyFlags(TransitionOptions, ETransitionOptions::SuspendGameplayTagTreeBroadcast))
-			{
-				SuspendGameplayTagTreeBroadcast.Emplace(!IsEngineExitRequested());
-			}
-
-			if (EnumHasAnyFlags(TransitionOptions, ETransitionOptions::SuspendAssetRegistry))
-			{
-				SuspendAssetRegistry.Emplace();
-			}
-		}
-
-		~FTransitionContenxt()
-		{
-			CompleteDelegate.ExecuteIfBound(Results);
-			// Explicit order of destruction is required
-			// At the very least we need to execute PluginManger batch last
-			// as it will run GC and does the leaked packages verification.
-			// And by that point all other resources need to be properly cleared,
-			// i.e. other batches need to finish.
-			SuspendAssetRegistry.Reset();
-			SuspendGameplayTagTreeBroadcast.Reset();
-			SuspendPluginManager.Reset();
-		}
-	};
-}
-
 void UGameFeaturesSubsystem::DeactivateGameFeaturePlugin(const FString& PluginURL)
 {
 	DeactivateGameFeaturePlugin(PluginURL, FGameFeaturePluginDeactivateComplete());
