Clean up game feature subsystem handling of tag/asset manager so it always uses the versions that are safe to call during shutdown.
Add some checks to the tag manager to catch future issues like this.

diff --git a/Engine/Plugins/Runtime/GameFeatures/Source/GameFeatures/Private/GameFeaturesSubsystem.cpp b/Engine/Plugins/Runtime/GameFeatures/Source/GameFeatures/Private/GameFeaturesSubsystem.cpp
--- a/Engine/Plugins/Runtime/GameFeatures/Source/GameFeatures/Private/GameFeaturesSubsystem.cpp
+++ b/Engine/Plugins/Runtime/GameFeatures/Source/GameFeatures/Private/GameFeaturesSubsystem.cpp
@@ -1538,12 +1538,7 @@ void UGameFeaturesSubsystem::LoadGameFeaturePlugin(TConstArrayView<FString> Plug
 
 		FLoadContext()
 		{
-			if (!IsEngineExitRequested())
-			{
-				UGameplayTagsManager::Get().PushDeferOnGameplayTagTreeChangedBroadcast();
-				bPushedTagsBroadcast = true;
-			}
-			else if(UGameplayTagsManager* TagsManager = UGameplayTagsManager::GetIfAllocated())
+			if (UGameplayTagsManager* TagsManager = UGameplayTagsManager::GetIfAllocated())
 			{
 				TagsManager->PushDeferOnGameplayTagTreeChangedBroadcast();
 				bPushedTagsBroadcast = true;
@@ -1639,12 +1634,7 @@ void UGameFeaturesSubsystem::RegisterGameFeaturePlugin(TConstArrayView<FString>
 
 		FLoadContext()
 		{
-			if (!IsEngineExitRequested())
-			{
-				UGameplayTagsManager::Get().PushDeferOnGameplayTagTreeChangedBroadcast();
-				bPushedTagsBroadcast = true;
-			}
-			else if (UGameplayTagsManager* TagsManager = UGameplayTagsManager::GetIfAllocated())
+			if (UGameplayTagsManager* TagsManager = UGameplayTagsManager::GetIfAllocated())
 			{
 				TagsManager->PushDeferOnGameplayTagTreeChangedBroadcast();
 				bPushedTagsBroadcast = true;
@@ -1797,12 +1787,7 @@ void UGameFeaturesSubsystem::ChangeGameFeatureTargetState(TConstArrayView<FStrin
 
 		FLoadContext()
 		{
-			if (!IsEngineExitRequested())
-			{
-				UGameplayTagsManager::Get().PushDeferOnGameplayTagTreeChangedBroadcast();
-				bPushedTagsBroadcast = true;
-			}
-			else if(UGameplayTagsManager* TagsManager = UGameplayTagsManager::GetIfAllocated())
+			if (UGameplayTagsManager* TagsManager = UGameplayTagsManager::GetIfAllocated())
 			{
 				TagsManager->PushDeferOnGameplayTagTreeChangedBroadcast();
 				bPushedTagsBroadcast = true;
@@ -2049,9 +2034,9 @@ namespace GameFeaturesSubsystem
 
 	struct FSuspendGameplayTagTreeChangedBroadcast
 	{
-		UE_NODISCARD_CTOR FSuspendGameplayTagTreeChangedBroadcast(bool bAllocateIfRequested)
+		UE_NODISCARD_CTOR FSuspendGameplayTagTreeChangedBroadcast()
 		{
-			if (UGameplayTagsManager* TagsManager = bAllocateIfRequested ? &UGameplayTagsManager::Get() : UGameplayTagsManager::GetIfAllocated() )
+			if (UGameplayTagsManager* TagsManager = UGameplayTagsManager::GetIfAllocated())
 			{
 				TagsManager->PushDeferOnGameplayTagTreeChangedBroadcast();
 				bWasSuspended = true;
@@ -2062,7 +2047,10 @@ namespace GameFeaturesSubsystem
 		{
 			if (bWasSuspended)
 			{
-				UGameplayTagsManager::Get().PopDeferOnGameplayTagTreeChangedBroadcast();
+				if (UGameplayTagsManager* TagsManager = UGameplayTagsManager::GetIfAllocated())
+				{
+					TagsManager->PopDeferOnGameplayTagTreeChangedBroadcast();
+				}
 			}
 		}
 
@@ -2101,7 +2089,7 @@ namespace GameFeaturesSubsystem
 
 			if (EnumHasAnyFlags(TransitionOptions, ETransitionOptions::SuspendGameplayTagTreeBroadcast))
 			{
-				SuspendGameplayTagTreeBroadcast.Emplace(!IsEngineExitRequested());
+				SuspendGameplayTagTreeBroadcast.Emplace();
 			}
 
 			if (EnumHasAnyFlags(TransitionOptions, ETransitionOptions::SuspendAssetRegistry))
@@ -2662,25 +2650,15 @@ void UGameFeaturesSubsystem::LoadBuiltInGameFeaturePluginsInternal(FBuiltInPlugi
 
 		FLoadContext()
 		{
-			if (!IsEngineExitRequested())
+			if (UAssetManager* AssetManager = UAssetManager::GetIfInitialized())
 			{
-				UAssetManager::Get().PushBulkScanning();
+				AssetManager->PushBulkScanning();
 				bPushedAssetBulkScanning = true;
-				UGameplayTagsManager::Get().PushDeferOnGameplayTagTreeChangedBroadcast();
-				bPushedTagsBroadcast = true;
 			}
-			else
+			if (UGameplayTagsManager* TagsManager = UGameplayTagsManager::GetIfAllocated())
 			{
-				if(UAssetManager* AssetManager =  UAssetManager::GetIfInitialized())
-				{
-					AssetManager->PushBulkScanning();
-					bPushedAssetBulkScanning = true;
-				}
-				if(UGameplayTagsManager* TagsManager = UGameplayTagsManager::GetIfAllocated())
-				{
-					TagsManager->PushDeferOnGameplayTagTreeChangedBroadcast();
-					bPushedTagsBroadcast = true;
-				}
+				TagsManager->PushDeferOnGameplayTagTreeChangedBroadcast();
+				bPushedTagsBroadcast = true;
 			}
 		}
 
@@ -2695,7 +2673,7 @@ void UGameFeaturesSubsystem::LoadBuiltInGameFeaturePluginsInternal(FBuiltInPlugi
 			}
 			if (bPushedAssetBulkScanning)
 			{
-				if(UAssetManager* AssetManager =  UAssetManager::GetIfInitialized())
+				if (UAssetManager* AssetManager = UAssetManager::GetIfInitialized())
 				{
 					AssetManager->PopBulkScanning();
 				}
diff --git a/Engine/Source/Runtime/GameplayTags/Private/GameplayTagsManager.cpp b/Engine/Source/Runtime/GameplayTags/Private/GameplayTagsManager.cpp
--- a/Engine/Source/Runtime/GameplayTags/Private/GameplayTagsManager.cpp
+++ b/Engine/Source/Runtime/GameplayTags/Private/GameplayTagsManager.cpp
@@ -351,6 +351,8 @@ void UGameplayTagsManager::LoadGameplayTagTables(bool bAllowAsyncLoad)
 	SCOPE_CYCLE_COUNTER(STAT_GameplayTags_LoadGameplayTags);
 
 	const UGameplayTagsSettings* Default = GetDefault<UGameplayTagsSettings>();
+	check(Default);
+
 	GameplayTagTables.Empty();
 
 #if !WITH_EDITOR
@@ -1164,7 +1166,10 @@ bool UGameplayTagsManager::ImportSingleGameplayTag(FGameplayTag& Tag, FName Impo
 
 UE_AUTORTFM_ALWAYS_OPEN void UGameplayTagsManager::InitializeManager()
 {
+	// This should never be called during shutdown, use GetIfAllocated
+	check(!IsEngineExitRequested());
 	check(!SingletonManager);
+
 	SCOPED_BOOT_TIMING("UGameplayTagsManager::InitializeManager");
 	SCOPE_LOG_TIME_IN_SECONDS(TEXT("UGameplayTagsManager::InitializeManager"), nullptr);
 
