UE: GameFeaturePlugin game policy can now determine if it is allowed to have a missing GameFeatureData. If the GFD is missing, a transient placeholder is created similar to when the plugin does not support content. By default it is not allowed (existing behavior).

#Synced-CL 48161255

diff --git a/Engine/Plugins/Runtime/GameFeatures/Source/GameFeatures/Private/GameFeaturePluginStateMachine.cpp b/Engine/Plugins/Runtime/GameFeatures/Source/GameFeatures/Private/GameFeaturePluginStateMachine.cpp
--- a/Engine/Plugins/Runtime/GameFeatures/Source/GameFeatures/Private/GameFeaturePluginStateMachine.cpp
+++ b/Engine/Plugins/Runtime/GameFeatures/Source/GameFeatures/Private/GameFeaturePluginStateMachine.cpp
@@ -3238,17 +3238,8 @@ struct FGameFeaturePluginState_Registering : public FGameFeaturePluginState
 				}
 				else
 				{
-					LoadGFDState = ELoadGFDState::Failed;
-
-					const FStringView ShortUrl = StateProperties.PluginIdentifier.GetIdentifyingString();
-					if (const TOptional<UE::UnifiedError::FError>& Error = GameFeatureDataHandle->GetError())
-					{
-						UE_LOG(LogGameFeatures, Error, TEXT("Game Feature Data loading failed with error [%s] for URL %.*s"), *LexToString(*Error), ShortUrl.Len(), ShortUrl.GetData());
-					}
-					else
-					{
-						UE_LOG(LogGameFeatures, Error, TEXT("Game Feature Data loading failed without an error for URL %.*s"), ShortUrl.Len(), ShortUrl.GetData());
-					}
+					// This will set LoadGFDState
+					HandleMissingGameFeatureData();
 				}
 				
 				UpdateStateMachineDeferred();
@@ -3276,6 +3267,44 @@ struct FGameFeaturePluginState_Registering : public FGameFeaturePluginState
 		return FGameFeaturePluginState::UseAsyncLoading();
 	}
 
+	UGameFeatureData* CreateTransientGameFeatureData()
+	{
+		return NewObject<UGameFeatureData>(GetTransientPackage(), FName(*StateProperties.PluginName), RF_Transient);
+	}
+
+	void HandleMissingGameFeatureData()
+	{
+		const FStringView ShortUrl = StateProperties.PluginIdentifier.GetIdentifyingString();
+		if (UGameFeaturesSubsystem::Get().GetPolicy().ShouldAllowMissingGameFeatureData(StateProperties.PluginName))
+		{
+			UE_LOG(LogGameFeatures, Log, TEXT("Game Feature Data package not found for URL %.*s. Creating a transient GameFeatureData as a placeholder."), ShortUrl.Len(), ShortUrl.GetData());
+			StateProperties.GameFeatureData = CreateTransientGameFeatureData();
+
+			LoadGFDState = ELoadGFDState::Success;
+		}
+		else
+		{
+			if (GameFeatureDataHandle.IsValid())
+			{
+				if (GameFeatureDataHandle->HasError())
+				{
+					const TOptional<UE::UnifiedError::FError>& Error = GameFeatureDataHandle->GetError();
+					UE_LOG(LogGameFeatures, Error, TEXT("Game Feature Data loading failed with error [%s] for URL %.*s"), *LexToString(*Error), ShortUrl.Len(), ShortUrl.GetData());
+				}
+				else
+				{
+					UE_LOG(LogGameFeatures, Error, TEXT("Game Feature Data loading failed without an error for URL %.*s"), ShortUrl.Len(), ShortUrl.GetData());
+				}
+			}
+			else
+			{
+				UE_LOG(LogGameFeatures, Error, TEXT("Game Feature Data package not found for URL %.*s"), ShortUrl.Len(), ShortUrl.GetData());
+			}
+
+			LoadGFDState = ELoadGFDState::Failed;
+		}
+	}
+
 	virtual void BeginState() override
 	{
 		TRACE_CPUPROFILER_EVENT_SCOPE(GFP_Registering_Begin);
@@ -3297,7 +3326,7 @@ struct FGameFeaturePluginState_Registering : public FGameFeaturePluginState
 		// If the plugin contains content then load the GameFeatureData otherwise procedurally create one that is transient.
 		if (!Plugin->GetDescriptor().bCanContainContent)
 		{
-			StateProperties.GameFeatureData = NewObject<UGameFeatureData>(GetTransientPackage(), FName(*StateProperties.PluginName), RF_Transient);
+			StateProperties.GameFeatureData = CreateTransientGameFeatureData();
 			LoadGFDState = ELoadGFDState::Success;
 			return;
 		}
@@ -3347,10 +3376,8 @@ struct FGameFeaturePluginState_Registering : public FGameFeaturePluginState
 			}
 			else
 			{
-				const FStringView ShortUrl = StateProperties.PluginIdentifier.GetIdentifyingString();
-				UE_LOG(LogGameFeatures, Error, TEXT("Game Feature Data package not found for URL %.*s"), ShortUrl.Len(), ShortUrl.GetData());
-
-				LoadGFDState = ELoadGFDState::Failed;
+				// This will set LoadGFDState
+				HandleMissingGameFeatureData();
 				return;
 			}
 		}
@@ -3387,17 +3414,8 @@ struct FGameFeaturePluginState_Registering : public FGameFeaturePluginState
 			}
 			else
 			{
-				LoadGFDState = ELoadGFDState::Failed;
-
-				const FStringView ShortUrl = StateProperties.PluginIdentifier.GetIdentifyingString();
-				if (const TOptional<UE::UnifiedError::FError>& Error = GameFeatureDataHandle->GetError())
-				{
-					UE_LOG(LogGameFeatures, Error, TEXT("Game Feature Data loading failed with error [%s] for URL %.*s"), *LexToString(*Error), ShortUrl.Len(), ShortUrl.GetData());
-				}
-				else
-				{
-					UE_LOG(LogGameFeatures, Error, TEXT("Game Feature Data loading failed without an error for URL %.*s"), ShortUrl.Len(), ShortUrl.GetData());
-				}
+				// This will set LoadGFDState
+				HandleMissingGameFeatureData();
 			}
 		}
 	}
diff --git a/Engine/Plugins/Runtime/GameFeatures/Source/GameFeatures/Public/GameFeaturesProjectPolicies.h b/Engine/Plugins/Runtime/GameFeatures/Source/GameFeatures/Public/GameFeaturesProjectPolicies.h
--- a/Engine/Plugins/Runtime/GameFeatures/Source/GameFeatures/Public/GameFeaturesProjectPolicies.h
+++ b/Engine/Plugins/Runtime/GameFeatures/Source/GameFeatures/Public/GameFeaturesProjectPolicies.h
@@ -82,6 +82,9 @@ public:
 	// Return true if a uplugin's details should be read and false if it should be skipped. Skipped plugins will not be processed as GFPs and skipped as though they didn't exist. Useful to limit the number of uplugin files opened for perf reasons
 	virtual bool ShouldReadPluginDetails(const FString& PluginDescriptorFilename) const { return true; }
 
+	// Return true if a plugin should be allowed to register if it doesn't have a GameFeatureData. Useful in cases where you want to allow the content of certain plugins to not exist.
+	virtual bool ShouldAllowMissingGameFeatureData(const FString& PluginName) const { return false; }
+
 	// Called to resolve plugin dependencies, will successfully return an empty string if a dependency is not a GFP.
 	// This may be called with file protocol for built-in plugins in some cases, even if a different protocol is used at runtime.
 	// returns The dependency URL or an error if the dependency could not be resolved
