Respond to GFP unregistration by removing cached profile item instances to allow proper GC.

#Synced-CL 48159805

diff --git a/Engine/Plugins/Runtime/GameFeatures/Source/GameFeatures/Private/GameFeaturesSubsystem.cpp b/Engine/Plugins/Runtime/GameFeatures/Source/GameFeatures/Private/GameFeaturesSubsystem.cpp
--- a/Engine/Plugins/Runtime/GameFeatures/Source/GameFeatures/Private/GameFeaturesSubsystem.cpp
+++ b/Engine/Plugins/Runtime/GameFeatures/Source/GameFeatures/Private/GameFeaturesSubsystem.cpp
@@ -2019,6 +2019,25 @@ bool UGameFeaturesSubsystem::IsGameFeaturePluginActiveByName(FStringView PluginN
 // Internal structs for suspending systems so we can better batch change GFPs from states with out spamming these systems
 namespace GameFeaturesSubsystem
 {
+	struct FNotifyObservers
+	{
+		UE_NODISCARD_CTOR FNotifyObservers()
+		{
+			if (GEngine)
+			{
+				UGameFeaturesSubsystem::Get().BatchCallbackObservers(UGameFeaturesSubsystem::EBatchObserverCallback::BatchOperationStarted);
+			}
+		}
+
+		~FNotifyObservers()
+		{
+			if (GEngine)
+			{
+				UGameFeaturesSubsystem::Get().BatchCallbackObservers(UGameFeaturesSubsystem::EBatchObserverCallback::BatchOperationComplete);
+			}
+		}
+	};
+
 	struct FSuspendPluginManger
 	{
 		UE_NODISCARD_CTOR FSuspendPluginManger()
@@ -2063,7 +2082,8 @@ namespace GameFeaturesSubsystem
 		SuspendPluginManager			= 1 << 0,
 		SuspendGameplayTagTreeBroadcast	= 1 << 1,
 		SuspendAssetRegistry			= 1 << 2,
-		All = SuspendPluginManager | SuspendGameplayTagTreeBroadcast | SuspendAssetRegistry,
+		NotifyObservers					= 1 << 3,
+		All = SuspendPluginManager | SuspendGameplayTagTreeBroadcast | SuspendAssetRegistry | NotifyObservers,
 	};
 	ENUM_CLASS_FLAGS(ETransitionOptions)
 
@@ -2076,6 +2096,7 @@ namespace GameFeaturesSubsystem
 		TOptional<FSuspendPluginManger> SuspendPluginManager;
 		TOptional<FSuspendGameplayTagTreeChangedBroadcast> SuspendGameplayTagTreeBroadcast;
 		TOptional<IAssetRegistry::FSuspendPathsDismounting> SuspendAssetRegistry;
+		TOptional<FNotifyObservers> NotifyObservers;
 
 		UE_NODISCARD_CTOR FTransitionContenxt(ETransitionOptions TransitionOptions = ETransitionOptions::All)
 		{
@@ -2096,6 +2117,11 @@ namespace GameFeaturesSubsystem
 			{
 				SuspendAssetRegistry.Emplace();
 			}
+
+			if (EnumHasAnyFlags(TransitionOptions, ETransitionOptions::NotifyObservers))
+			{
+				NotifyObservers.Emplace();
+			}
 		}
 
 		~FTransitionContenxt()
@@ -2106,6 +2132,7 @@ namespace GameFeaturesSubsystem
 			// as it will run GC and does the leaked packages verification.
 			// And by that point all other resources need to be properly cleared,
 			// i.e. other batches need to finish.
+			NotifyObservers.Reset(); // Some observers are responsible for clearing references; must happen before SuspendPluginManager.Reset()
 			SuspendAssetRegistry.Reset();
 			SuspendGameplayTagTreeBroadcast.Reset();
 			SuspendPluginManager.Reset();
@@ -4362,6 +4389,32 @@ void UGameFeaturesSubsystem::CallbackObservers(EObserverCallback CallbackType, c
 	}
 }
 
+void UGameFeaturesSubsystem::BatchCallbackObservers(EBatchObserverCallback BatchCallbackType)
+{
+	// Protect against modifying the observer list during iteration
+	TArray<UObject*> LocalObservers(Observers);
+
+	switch (BatchCallbackType)
+	{
+		case EBatchObserverCallback::BatchOperationStarted:
+		{
+			for (UObject* Observer : LocalObservers)
+			{
+				CastChecked<IGameFeatureStateChangeObserver>(Observer)->OnGameFeatureBatchOperationStarted();
+			}
+			break;
+		}
+		case EBatchObserverCallback::BatchOperationComplete:
+		{
+			for (UObject* Observer : LocalObservers)
+			{
+				CastChecked<IGameFeatureStateChangeObserver>(Observer)->OnGameFeatureBatchOperationComplete();
+			}
+			break;
+		}
+	}
+}
+
 void UGameFeaturesSubsystem::RegisterRunningStateMachine(UGameFeaturePluginStateMachine* GFPSM)
 {
 	check(!RunningStateMachines.Contains(GFPSM));
diff --git a/Engine/Plugins/Runtime/GameFeatures/Source/GameFeatures/Public/GameFeatureStateChangeObserver.h b/Engine/Plugins/Runtime/GameFeatures/Source/GameFeatures/Public/GameFeatureStateChangeObserver.h
--- a/Engine/Plugins/Runtime/GameFeatures/Source/GameFeatures/Public/GameFeatureStateChangeObserver.h
+++ b/Engine/Plugins/Runtime/GameFeatures/Source/GameFeatures/Public/GameFeatureStateChangeObserver.h
@@ -83,4 +83,8 @@ public:
 		EX: Downloading paused due to a users cellular data settings or the user taking a pause action. We
 		may not yet transition to a download error, but want a way to observe this behavior. */
 	virtual void OnGameFeaturePauseChange(const FString& PluginURL, const FString& PluginName, FGameFeaturePauseStateChangeContext& Context) {}
+
+	virtual void OnGameFeatureBatchOperationStarted() {}
+
+	virtual void OnGameFeatureBatchOperationComplete() {}
 };
diff --git a/Engine/Plugins/Runtime/GameFeatures/Source/GameFeatures/Public/GameFeaturesSubsystem.h b/Engine/Plugins/Runtime/GameFeatures/Source/GameFeatures/Public/GameFeaturesSubsystem.h
--- a/Engine/Plugins/Runtime/GameFeatures/Source/GameFeatures/Public/GameFeaturesSubsystem.h
+++ b/Engine/Plugins/Runtime/GameFeatures/Source/GameFeatures/Public/GameFeaturesSubsystem.h
@@ -34,6 +34,11 @@ enum class EInstallBundleResult : uint32;
 enum class EInstallBundleRequestFlags : uint32;
 enum class EInstallBundleReleaseRequestFlags : uint32;
 
+namespace GameFeaturesSubsystem
+{
+	struct FNotifyObservers;
+}
+
 /** Holds static global information about how our PluginURLs are structured */
 namespace UE::GameFeatures
 {
@@ -923,6 +928,15 @@ private:
 		const UGameFeatureData* GameFeatureData = nullptr, 
 		FGameFeatureStateChangeContext* StateChangeContext = nullptr);
 
+	enum class EBatchObserverCallback : uint8
+	{
+		BatchOperationStarted,
+		BatchOperationComplete
+	};
+
+	UE_API void BatchCallbackObservers(EBatchObserverCallback BatchCallbackType);
+	friend struct GameFeaturesSubsystem::FNotifyObservers;
+
 	/** Registers a state machine that is in transition and running */
 	UE_API void RegisterRunningStateMachine(UGameFeaturePluginStateMachine* GFPSM);
 
