Add support for explicitly loading a list of GFPs in an amortized way from an explicit list

#Synced-CL 48202912

diff --git a/Engine/Plugins/Runtime/GameFeatures/Source/GameFeatures/Private/GameFeaturesSubsystem.cpp b/Engine/Plugins/Runtime/GameFeatures/Source/GameFeatures/Private/GameFeaturesSubsystem.cpp
--- a/Engine/Plugins/Runtime/GameFeatures/Source/GameFeatures/Private/GameFeaturesSubsystem.cpp
+++ b/Engine/Plugins/Runtime/GameFeatures/Source/GameFeatures/Private/GameFeaturesSubsystem.cpp
@@ -2647,18 +2647,40 @@ private:
 
 #endif // !UE_BUILD_SHIPPING
 
+TArray<TSharedRef<IPlugin>> UGameFeaturesSubsystem::GetInternalEnabledPlugins() const
+{
+	TArray<TSharedRef<IPlugin>> EnabledPlugins = IPluginManager::Get().GetEnabledPlugins();
+
+	if (UE::GameFeatures::bTrimNonStartupEnabledPlugins)
+	{
+		const TSet<FString>& CompiledInPlugins = IPluginManager::Get().GetPluginsEnabledAtStartup();
+		EnabledPlugins.RemoveAllSwap([&CompiledInPlugins](const TSharedRef<IPlugin>& Plugin)
+			{
+				const bool bRemove = !CompiledInPlugins.Contains(Plugin->GetName());
+				return bRemove;
+			});
+	}
+
+	return EnabledPlugins;
+}
+
 void UGameFeaturesSubsystem::LoadBuiltInGameFeaturePlugins(FBuiltInPluginAdditionalFilters AdditionalFilter, const FBuiltInGameFeaturePluginsLoaded& InCompleteDelegate /*= FBuiltInGameFeaturePluginsLoaded()*/)
 {
 	FBuiltInPluginAdditionalFilters_Copyable NullCopyable;
-	LoadBuiltInGameFeaturePluginsInternal(AdditionalFilter, NullCopyable, 0, InCompleteDelegate);
+	LoadBuiltInGameFeaturePluginsInternal(GetInternalEnabledPlugins(), AdditionalFilter, NullCopyable, 0, InCompleteDelegate);
 }
 
 void UGameFeaturesSubsystem::LoadBuiltInGameFeaturePlugins_Amortized(const FBuiltInPluginAdditionalFilters_Copyable& AdditionalFilter_Copyable, int32 AmortizeRate, const FBuiltInGameFeaturePluginsLoaded& InCompleteDelegate /*= FBuiltInGameFeaturePluginsLoaded()*/)
 {
-	LoadBuiltInGameFeaturePluginsInternal(AdditionalFilter_Copyable, AdditionalFilter_Copyable, AmortizeRate, InCompleteDelegate);
+	LoadBuiltInGameFeaturePluginsInternal(GetInternalEnabledPlugins(), AdditionalFilter_Copyable, AdditionalFilter_Copyable, AmortizeRate, InCompleteDelegate);
 }
 
-void UGameFeaturesSubsystem::LoadBuiltInGameFeaturePluginsInternal(FBuiltInPluginAdditionalFilters AdditionalFilter, const FBuiltInPluginAdditionalFilters_Copyable& AdditionalFilter_Copyable, int32 AmortizeRate, const FBuiltInGameFeaturePluginsLoaded& InCompleteDelegate /*= FBuiltInGameFeaturePluginsLoaded()*/)
+void UGameFeaturesSubsystem::LoadBuiltInGameFeaturePlugins_Amortized(const TArray<TSharedRef<IPlugin>>& Plugins, const FBuiltInPluginAdditionalFilters_Copyable& AdditionalFilter_Copyable, int32 AmortizeRate, const FBuiltInGameFeaturePluginsLoaded& InCompleteDelegate)
+{
+	LoadBuiltInGameFeaturePluginsInternal(Plugins, AdditionalFilter_Copyable, AdditionalFilter_Copyable, AmortizeRate, InCompleteDelegate);
+}
+
+void UGameFeaturesSubsystem::LoadBuiltInGameFeaturePluginsInternal(TArray<TSharedRef<IPlugin>> Plugins, FBuiltInPluginAdditionalFilters AdditionalFilter, const FBuiltInPluginAdditionalFilters_Copyable& AdditionalFilter_Copyable, int32 AmortizeRate, const FBuiltInGameFeaturePluginsLoaded& InCompleteDelegate /*= FBuiltInGameFeaturePluginsLoaded()*/)
 {
 	TRACE_CPUPROFILER_EVENT_SCOPE_STR("UGameFeaturesSubsystem::LoadBuiltInGameFeaturePlugins");
 
@@ -2713,25 +2735,14 @@ void UGameFeaturesSubsystem::LoadBuiltInGameFeaturePluginsInternal(FBuiltInPlugi
 	LoadContext->CompleteDelegate = InCompleteDelegate;
 
 	FBuiltInPluginLoadTimeTracker PluginLoadTimeTracker;
-	TArray<TSharedRef<IPlugin>> EnabledPlugins = IPluginManager::Get().GetEnabledPlugins();
 
-	if (UE::GameFeatures::bTrimNonStartupEnabledPlugins)
-	{
-		const TSet<FString>& CompiledInPlugins = IPluginManager::Get().GetPluginsEnabledAtStartup();
-		EnabledPlugins.RemoveAllSwap([&CompiledInPlugins](const TSharedRef<IPlugin>& Plugin)
-			{
-				const bool bRemove = !CompiledInPlugins.Contains(Plugin->GetName());
-				return bRemove;
-			});
-	}
-
-	LoadContext->Results.Reserve(EnabledPlugins.Num());
-	for (const TSharedRef<IPlugin>& Plugin : EnabledPlugins)
+	LoadContext->Results.Reserve(Plugins.Num());
+	for (const TSharedRef<IPlugin>& Plugin : Plugins)
 	{
 		LoadContext->Results.Add(Plugin->GetName(), MakeError("Pending"));
 	}
 
-	const int32 NumPluginsToLoad = EnabledPlugins.Num();
+	const int32 NumPluginsToLoad = Plugins.Num();
 	UE_LOG(LogGameFeatures, Log, TEXT("Loading %i builtins. %s"), NumPluginsToLoad, AmortizeRate > 0 ? TEXT("Amortized") : TEXT("Not Amortized"));
 
 	// Sort the plugins so we can more accurately track how long it takes to load rather than have inconsistent dependency timings.
@@ -2741,7 +2752,7 @@ void UGameFeaturesSubsystem::LoadBuiltInGameFeaturePluginsInternal(FBuiltInPlugi
 	{
 		IPluginManager& PluginManager = IPluginManager::Get();
 		Dependencies.Reset();
-		
+
 		const FPluginDescriptor& Desc = CurrentPlugin->GetDescriptor();
 		for (const FPluginReferenceDescriptor& Dependency : Desc.Plugins)
 		{
@@ -2756,7 +2767,7 @@ void UGameFeaturesSubsystem::LoadBuiltInGameFeaturePluginsInternal(FBuiltInPlugi
 		return Dependencies;
 	};
 
-	Algo::TopologicalSort(EnabledPlugins, GetPluginDependencies);
+	Algo::TopologicalSort(Plugins, GetPluginDependencies);
 
 	auto ProcessGameFeaturePlugin = [this](const TSharedRef<IPlugin>& PluginToProcess, const TSharedRef<FLoadContext>& LoadContextToProcess, const FBuiltInPluginAdditionalFilters& AdditionalFilterForProcess)
 	{
@@ -2771,7 +2782,7 @@ void UGameFeaturesSubsystem::LoadBuiltInGameFeaturePluginsInternal(FBuiltInPlugi
 	if (AmortizeRate > 0)
 	{
 		LoadContext->PluginAmortizeRate = AmortizeRate;
-		LoadContext->AmortizedPluginsRemaining = EnabledPlugins;
+		LoadContext->AmortizedPluginsRemaining = Plugins;
 		// Note that we are using AdditionalFilter_Copyable here. We need to make a copy of thus function so that the captured properties in lambdas are not destroyed
 		FTSTicker::GetCoreTicker().AddTicker(FTickerDelegate::CreateWeakLambda(this, [this, LoadContext, AdditionalFilter_Copyable, ProcessGameFeaturePlugin](float)
 		{
@@ -2793,7 +2804,7 @@ void UGameFeaturesSubsystem::LoadBuiltInGameFeaturePluginsInternal(FBuiltInPlugi
 	}
 	else
 	{
-		for (const TSharedRef<IPlugin>& Plugin : EnabledPlugins)
+		for (const TSharedRef<IPlugin>& Plugin : Plugins)
 		{
 			FBuiltInPluginLoadTimeTrackerScope TrackerScope(PluginLoadTimeTracker, Plugin);
 			ProcessGameFeaturePlugin(Plugin, LoadContext, AdditionalFilter);
diff --git a/Engine/Plugins/Runtime/GameFeatures/Source/GameFeatures/Public/GameFeaturesSubsystem.h b/Engine/Plugins/Runtime/GameFeatures/Source/GameFeatures/Public/GameFeaturesSubsystem.h
--- a/Engine/Plugins/Runtime/GameFeatures/Source/GameFeatures/Public/GameFeaturesSubsystem.h
+++ b/Engine/Plugins/Runtime/GameFeatures/Source/GameFeatures/Public/GameFeaturesSubsystem.h
@@ -702,8 +702,10 @@ public:
 
 	/** Loads all built-in game feature plugins that pass the specified filters, split over multiple frames processing only AmortizeRate plugins per frame if greater than 0. Note that AdditionalFilter must be a TFunction and not a TFunctionRef since it will be used in future ticks */
 	UE_API void LoadBuiltInGameFeaturePlugins_Amortized(const FBuiltInPluginAdditionalFilters_Copyable& AdditionalFilter_Copyable, int32 AmortizeRate, const FBuiltInGameFeaturePluginsLoaded& CompleteDelegate = FBuiltInGameFeaturePluginsLoaded());
+	UE_API void LoadBuiltInGameFeaturePlugins_Amortized(const TArray<TSharedRef<IPlugin>>& Plugins, const FBuiltInPluginAdditionalFilters_Copyable& AdditionalFilter, int32 AmortizeRate, const FBuiltInGameFeaturePluginsLoaded& CompleteDelegate = FBuiltInGameFeaturePluginsLoaded());
 private:
-	void LoadBuiltInGameFeaturePluginsInternal(FBuiltInPluginAdditionalFilters AdditionalFilter, const FBuiltInPluginAdditionalFilters_Copyable& AdditionalFilter_Copyable, int32 AmortizeRate, const FBuiltInGameFeaturePluginsLoaded& CompleteDelegate = FBuiltInGameFeaturePluginsLoaded());
+	TArray<TSharedRef<IPlugin>> GetInternalEnabledPlugins() const;
+	void LoadBuiltInGameFeaturePluginsInternal(TArray<TSharedRef<IPlugin>> Plugins, FBuiltInPluginAdditionalFilters AdditionalFilter, const FBuiltInPluginAdditionalFilters_Copyable& AdditionalFilter_Copyable, int32 AmortizeRate, const FBuiltInGameFeaturePluginsLoaded& CompleteDelegate = FBuiltInGameFeaturePluginsLoaded());
 public:
 
 	/** Returns the list of plugin filenames that have progressed beyond installed. Used in cooking to determine which will be cooked. */
