Fix bug where data assets made from blueprints had the NativeClass asset registry tag set to the blueprint instead of the parent native class.
This will be fixed the next time those assets are resaved.
Save 8 bytes of memory for every data asset by removing the private variable and deleting some deprecated code.
#Synced-CL 48359414

diff --git a/Engine/Plugins/Experimental/CommonConversation/Source/CommonConversationRuntime/Private/ConversationDatabase.cpp b/Engine/Plugins/Experimental/CommonConversation/Source/CommonConversationRuntime/Private/ConversationDatabase.cpp
--- a/Engine/Plugins/Experimental/CommonConversation/Source/CommonConversationRuntime/Private/ConversationDatabase.cpp
+++ b/Engine/Plugins/Experimental/CommonConversation/Source/CommonConversationRuntime/Private/ConversationDatabase.cpp
@@ -25,18 +25,6 @@ UConversationDatabase::UConversationDatabase(const FObjectInitializer& ObjectIni
 {
 }
 
-void UConversationDatabase::GetAssetRegistryTags(TArray<FAssetRegistryTag>& OutTags) const
-{
-	PRAGMA_DISABLE_DEPRECATION_WARNINGS;
-	Super::GetAssetRegistryTags(OutTags);
-	PRAGMA_ENABLE_DEPRECATION_WARNINGS;
-}
-
-void UConversationDatabase::GetAssetRegistryTags(FAssetRegistryTagsContext Context) const
-{
-	Super::GetAssetRegistryTags(Context);
-}
-
 #if WITH_EDITOR
 void UConversationDatabase::PostDuplicate(bool bDuplicateForPIE)
 {
diff --git a/Engine/Plugins/Experimental/CommonConversation/Source/CommonConversationRuntime/Public/ConversationDatabase.h b/Engine/Plugins/Experimental/CommonConversation/Source/CommonConversationRuntime/Public/ConversationDatabase.h
--- a/Engine/Plugins/Experimental/CommonConversation/Source/CommonConversationRuntime/Public/ConversationDatabase.h
+++ b/Engine/Plugins/Experimental/CommonConversation/Source/CommonConversationRuntime/Public/ConversationDatabase.h
@@ -78,10 +78,6 @@ class UConversationDatabase : public UPrimaryDataAsset
 public:
 	UE_API UConversationDatabase(const FObjectInitializer& ObjectInitializer);
 
-	UE_API virtual void GetAssetRegistryTags(FAssetRegistryTagsContext Context) const override;
-	UE_DEPRECATED(5.4, "Implement the version that takes FAssetRegistryTagsContext instead.")
-	UE_API virtual void GetAssetRegistryTags(TArray<FAssetRegistryTag>& OutTags) const override;
-
 #if WITH_EDITOR
 	UE_API virtual void PostDuplicate(bool bDuplicateForPIE) override;
 	UE_API virtual EDataValidationResult IsDataValid(class FDataValidationContext& Context) const override;
diff --git a/Engine/Plugins/Runtime/GameFeatures/Source/GameFeatures/Private/GameFeatureData.cpp b/Engine/Plugins/Runtime/GameFeatures/Source/GameFeatures/Private/GameFeatureData.cpp
--- a/Engine/Plugins/Runtime/GameFeatures/Source/GameFeatures/Private/GameFeatureData.cpp
+++ b/Engine/Plugins/Runtime/GameFeatures/Source/GameFeatures/Private/GameFeatureData.cpp
@@ -867,13 +867,6 @@ void UGameFeatureData::GetDependencyDirectoriesFromAssetData(const FAssetData& A
 	}
 }
 
-void UGameFeatureData::GetAssetRegistryTags(TArray<FAssetRegistryTag>& OutTags) const
-{
-	PRAGMA_DISABLE_DEPRECATION_WARNINGS;
-	Super::GetAssetRegistryTags(OutTags);
-	PRAGMA_ENABLE_DEPRECATION_WARNINGS;
-}
-
 void UGameFeatureData::GetAssetRegistryTags(FAssetRegistryTagsContext Context) const
 {
 	Super::GetAssetRegistryTags(Context);
diff --git a/Engine/Plugins/Runtime/GameFeatures/Source/GameFeatures/Public/GameFeatureData.h b/Engine/Plugins/Runtime/GameFeatures/Source/GameFeatures/Public/GameFeatureData.h
--- a/Engine/Plugins/Runtime/GameFeatures/Source/GameFeatures/Public/GameFeatureData.h
+++ b/Engine/Plugins/Runtime/GameFeatures/Source/GameFeatures/Public/GameFeatureData.h
@@ -52,15 +52,6 @@ public:
 	virtual TArray<FPrimaryAssetTypeInfo>& GetPrimaryAssetTypesToScan() { return PrimaryAssetTypesToScan; }
 	UE_API virtual void GetAssetRegistryTags(FAssetRegistryTagsContext Context) const override;
 	static UE_API void GetDependencyDirectoriesFromAssetData(const FAssetData& AssetData, TArray<FString>& OutDependencyDirectories);
-
-	//~Begin deprecation
-	UE_DEPRECATED(5.4, "Implement the version that takes FAssetRegistryTagsContext instead.")
-	UE_API virtual void GetAssetRegistryTags(TArray<FAssetRegistryTag>& OutTags) const override;
-	UE_DEPRECATED(5.4, "GetContentBundleGuidsAssetRegistryTag is deprecated")
-	static FName GetContentBundleGuidsAssetRegistryTag() { return NAME_None; }
-	UE_DEPRECATED(5.4, "GetContentBundleGuidsFromAsset is deprecated, use GetDependencyDirectoriesFromAssetData")
-	static void GetContentBundleGuidsFromAsset(const FAssetData& Asset, TArray<FGuid>& OutContentBundleGuids) {}
-	//~End deprecation
 #endif //if WITH_EDITOR
 
 	/** Method to process the base ini file for the plugin during loading */
diff --git a/Engine/Source/Runtime/Engine/Classes/Engine/DataAsset.h b/Engine/Source/Runtime/Engine/Classes/Engine/DataAsset.h
--- a/Engine/Source/Runtime/Engine/Classes/Engine/DataAsset.h
+++ b/Engine/Source/Runtime/Engine/Classes/Engine/DataAsset.h
@@ -19,16 +19,16 @@
 UCLASS(abstract, MinimalAPI, Meta = (LoadBehavior = "LazyOnDemand"))
 class UDataAsset : public UObject
 {
-	GENERATED_UCLASS_BODY()
+	GENERATED_BODY()
 public:
 	// UObject interface
+	ENGINE_API virtual void GetAssetRegistryTags(FAssetRegistryTagsContext Context) const override;
 #if WITH_EDITORONLY_DATA
 	ENGINE_API virtual void Serialize(FStructuredArchiveRecord Record) override;
 #endif
 
-private:
-	UPROPERTY(AssetRegistrySearchable)
-	TSubclassOf<UDataAsset> NativeClass;
+	/** Asset registry tag NativeClass with full path in export form Class'/Script/ModuleName.ClassName' of the first found parent native class */
+	static ENGINE_API const FName NativeClassTag;
 };
 
 /**
diff --git a/Engine/Source/Runtime/Engine/Private/DataAsset.cpp b/Engine/Source/Runtime/Engine/Private/DataAsset.cpp
--- a/Engine/Source/Runtime/Engine/Private/DataAsset.cpp
+++ b/Engine/Source/Runtime/Engine/Private/DataAsset.cpp
@@ -3,16 +3,28 @@
 #include "Engine/DataAsset.h"
 
 #include "Misc/PackageName.h"
+#include "UObject/AssetRegistryTagsContext.h"
 #include "UObject/ObjectSaveContext.h"
 #include "UObject/Package.h"
 #include "Engine/AssetManager.h"
 
 #include UE_INLINE_GENERATED_CPP_BY_NAME(DataAsset)
 
-UDataAsset::UDataAsset(const FObjectInitializer& ObjectInitializer)
-	: Super(ObjectInitializer)
+const FName UDataAsset::NativeClassTag(TEXT("NativeClass"));
+
+void UDataAsset::GetAssetRegistryTags(FAssetRegistryTagsContext Context) const
 {
-	NativeClass = GetClass();
+	// Walk up until we find a native class (ie 'while they are BP classes')
+	UClass* NativeParentClass = GetClass();
+	while (!NativeParentClass->HasAnyClassFlags(CLASS_Native | CLASS_Intrinsic))
+	{
+		NativeParentClass = NativeParentClass->GetSuperClass();
+	}
+	check(NativeParentClass);
+
+	Context.AddTag(FAssetRegistryTag(NativeClassTag, FObjectPropertyBase::GetExportPath(NativeParentClass), FAssetRegistryTag::TT_Alphabetical));
+
+	Super::GetAssetRegistryTags(Context);
 }
 
 #if WITH_EDITORONLY_DATA
