GFP Refactor: Remove all methods from UGameFeaturePluginStateMachine and move them to UGameFeaturesSubsystem.
The goal is to progressively hoist code from UGameFeaturePluginStateMachine into UGameFeaturesSubsystem so that the subsystem can operate on multiple GFPs making the same state transition at the same time, allowing more batching & parallelism and allowing us to remove redundant work.
Initial changes will just move the existing methods so that access to per-GFP state becomes explicit via a parameter.

#Synced-CL Unknown

diff --git a/Engine/Plugins/Runtime/GameFeatures/Source/GameFeatures/Private/GameFeaturePluginStateMachine.cpp b/Engine/Plugins/Runtime/GameFeatures/Source/GameFeatures/Private/GameFeaturePluginStateMachine.cpp
--- a/Engine/Plugins/Runtime/GameFeatures/Source/GameFeatures/Private/GameFeaturePluginStateMachine.cpp
+++ b/Engine/Plugins/Runtime/GameFeatures/Source/GameFeatures/Private/GameFeaturePluginStateMachine.cpp
@@ -5,6 +5,7 @@
 #include "AssetRegistry/AssetRegistryState.h"
 #include "Engine/StreamableManager.h"
 #include "GameFeatureData.h"
+#include "GameFeaturesSubsystem.h"
 #include "GameplayTagsManager.h"
 #include "Engine/AssetManager.h"
 #include "GenericPlatform/GenericPlatformMisc.h"
@@ -62,6 +63,9 @@
 
 LLM_DEFINE_TAG(GFP, TEXT("GameFeatures"), TEXT("UObject"));
 
+#define GFP_STATE_LOG(VERBOSITY, MSG, ...) UE_LOGFMT(LogGameFeatures, VERBOSITY, \
+	"{Plugin}[{State}] " MSG, StateProperties.PluginIdentifier.GetPluginName(), State __VA_OPT__(,) __VA_ARGS__);
+
 namespace UE::GameFeatures
 {
 	static const FString StateMachineErrorNamespace(TEXT("GameFeaturePlugin.StateMachine."));
@@ -129,13 +133,13 @@ namespace UE::GameFeatures
 	void MountLocalizationData(UGameFeaturePluginStateMachine* CurrentMachine, FGameFeaturePluginStateMachineProperties& StateProperties)
 	{
 		check(IsInGameThread());
-		check(CurrentMachine && &CurrentMachine->GetProperties() == &StateProperties);
+		check(CurrentMachine && &CurrentMachine->StateProperties == &StateProperties);
 
 		StateProperties.bIsLoadingLocalizationData = true;
 		IPluginManager::Get().MountExplicitlyLoadedPluginLocalizationData(StateProperties.PluginName, [WeakMachine = TWeakObjectPtr<UGameFeaturePluginStateMachine>(CurrentMachine), &StateProperties, bAllowAsyncLoading = StateProperties.AllowAsyncLoading()](bool bLoadedLocalization, const FString& PluginName)
 		{
 			if (UGameFeaturePluginStateMachine* StateMachine = WeakMachine.Get();
-				StateMachine && ensureAlways(&StateMachine->GetProperties() == &StateProperties))
+				StateMachine && ensureAlways(&StateMachine->StateProperties == &StateProperties))
 			{
 				if (IsInGameThread())
 				{
@@ -146,7 +150,7 @@ namespace UE::GameFeatures
 					ExecuteOnGameThread(UE_SOURCE_LOCATION, [WeakMachine, &StateProperties]()
 					{
 						if (UGameFeaturePluginStateMachine* StateMachine = WeakMachine.Get();
-							StateMachine && ensureAlways(&StateMachine->GetProperties() == &StateProperties))
+							StateMachine && ensureAlways(&StateMachine->StateProperties == &StateProperties))
 						{
 							StateProperties.bIsLoadingLocalizationData = false;
 							StateProperties.OnRequestUpdateStateMachine.ExecuteIfBound();
@@ -606,6 +610,16 @@ UE::GameFeatures::FResult FGameFeaturePluginState::GetErrorResult(const FString&
 	return MoveTemp(BaseResult);
 }
 
+void SerializeForLog(FCbWriter& Writer, EGameFeaturePluginState State)
+{
+#define STATE_TO_STRING(NAME, LOCNAME) case EGameFeaturePluginState::NAME: Writer.AddString(TEXT(#NAME)); break;
+	switch (State)
+	{
+GAME_FEATURE_PLUGIN_STATE_LIST(STATE_TO_STRING)
+	}
+#undef STATE_TO_STRING
+	
+}
 
 FGameFeaturePluginState::~FGameFeaturePluginState()
 {
@@ -652,10 +666,19 @@ void FGameFeaturePluginState::UpdateStateMachineDeferred(float Delay /*= 0.0f*/)
 {
 	CleanupDeferredUpdateCallbacks();
 
+	GFP_STATE_LOG(VeryVerbose, "queueing deferred update");
 	TickHandle = FTSTicker::GetCoreTicker().AddTicker(FTickerDelegate::CreateLambda([this](float dts) mutable
 	{
 		// @note Release FGameFeaturePluginState::TickHandle first in case the termination callback triggers a GC and destroys the state machine
 		TickHandle.Reset();
+		if (StateProperties.OnRequestUpdateStateMachine.IsBound())
+		{
+			StateProperties.OnRequestUpdateStateMachine.Execute();
+		}
+		else
+		{
+			GFP_STATE_LOG(Warning, "OnRequestUpdateStateMachine not bound");
+		}
 		StateProperties.OnRequestUpdateStateMachine.ExecuteIfBound();
 		return false;
 	}), Delay);
@@ -753,9 +776,7 @@ bool FGameFeaturePluginState::UseAsyncLoading() const
 template<typename TransitionPolicy>
 struct FTransitionDependenciesGameFeaturePluginState : public FGameFeaturePluginState
 {
-	FTransitionDependenciesGameFeaturePluginState(FGameFeaturePluginStateMachineProperties& InStateProperties)
-		: FGameFeaturePluginState(InStateProperties)
-	{}
+	using FGameFeaturePluginState::FGameFeaturePluginState;
 
 	virtual ~FTransitionDependenciesGameFeaturePluginState()
 	{
@@ -808,17 +829,17 @@ struct FTransitionDependenciesGameFeaturePluginState : public FGameFeaturePlugin
 			RemainingDependencies.Reserve(Dependencies.Num());
 			for (UGameFeaturePluginStateMachine* Dependency : Dependencies)
 			{
-				ensureMsgf(bAllowAsyncLoading || !Dependency->AllowAsyncLoading(), 
+				ensureMsgf(bAllowAsyncLoading || !GameFeaturesSubsystem.AllowAsyncLoading(Dependency), 
 					TEXT("FGameFeaturePluginState::AllowAsyncLoading is false for %s but true for dependency being waited on %s"), 
-					*StateProperties.PluginName, *Dependency->GetPluginURL());
+					*StateProperties.PluginName, *GameFeaturesSubsystem.GetPluginURL(Dependency));
 
 				if (TransitionPolicy::ShouldWaitForDependencies() && 
-					Dependency->IsInErrorState() &&
-					Dependency->IsErrorStateUnrecoverable())
+					GameFeaturesSubsystem.IsInErrorState(Dependency) &&
+					GameFeaturesSubsystem.IsErrorStateUnrecoverable(Dependency))
 				{
-					const EGameFeaturePluginState DependencyState = Dependency->GetCurrentState();
+					const EGameFeaturePluginState DependencyState = GameFeaturesSubsystem.GetCurrentState(Dependency);
 					const FStringView ShortUrl = StateProperties.PluginIdentifier.GetIdentifyingString();
-					const FStringView DepShortUrl = Dependency->GetPluginIdentifier().GetIdentifyingString();
+					const FStringView DepShortUrl = GameFeaturesSubsystem.GetPluginIdentifier(Dependency).GetIdentifyingString();
 					UE_LOG(LogGameFeatures, Error, TEXT("Plugin (%.*s) dependency (%.*s) failed and is in state %s"),
 						ShortUrl.Len(), ShortUrl.GetData(),
 						DepShortUrl.Len(), DepShortUrl.GetData(),
@@ -829,7 +850,7 @@ struct FTransitionDependenciesGameFeaturePluginState : public FGameFeaturePlugin
 						UE_LOG(LogGameFeatures, Error, TEXT("Setting plugin (%.*s) to be in unrecoverable error because dependency (%.*s) is in unrecoverable error"), 
 							ShortUrl.Len(), ShortUrl.GetData(),
 							DepShortUrl.Len(), DepShortUrl.GetData());
-						CurrentMachine->SetUnrecoverableError();
+						GameFeaturesSubsystem.SetUnrecoverableError(CurrentMachine);
 					}
 					return;
 				}
@@ -854,20 +875,20 @@ struct FTransitionDependenciesGameFeaturePluginState : public FGameFeaturePlugin
 			if (Pair.Value.HasError())
 			{
 				const FStringView ShortUrl = StateProperties.PluginIdentifier.GetIdentifyingString();
-				const FStringView DepShortUrl = RemainingDependency->GetPluginIdentifier().GetIdentifyingString();
+				const FStringView DepShortUrl = GameFeaturesSubsystem.GetPluginIdentifier(RemainingDependency).GetIdentifyingString();
 				UE_LOG(LogGameFeatures, Error, TEXT("Plugin (%.*s) dependency (%.*s) failed to transition with error %s"),
 					ShortUrl.Len(), ShortUrl.GetData(),
 					DepShortUrl.Len(), DepShortUrl.GetData(),
 					*Pair.Value.GetError());
 				StateStatus.SetTransitionError(TransitionPolicy::GetErrorState(), GetErrorResult(TEXT("Failed_Dependency_Transition")));
-				if (RemainingDependency->IsErrorStateUnrecoverable())
+				if (GameFeaturesSubsystem.IsErrorStateUnrecoverable(RemainingDependency))
 				{
 					if (UGameFeaturePluginStateMachine* CurrentMachine = GameFeaturesSubsystem.FindGameFeaturePluginStateMachine(StateProperties.PluginIdentifier))
 					{
 						UE_LOG(LogGameFeatures, Error, TEXT("Setting plugin (%.*s) to be in unrecoverable error because dependency (%.*s) is in unrecoverable error"),
 							ShortUrl.Len(), ShortUrl.GetData(),
 							DepShortUrl.Len(), DepShortUrl.GetData());
-						CurrentMachine->SetUnrecoverableError();
+						GameFeaturesSubsystem.SetUnrecoverableError(CurrentMachine);
 					}
 				}
 				return;
@@ -884,19 +905,20 @@ struct FTransitionDependenciesGameFeaturePluginState : public FGameFeaturePlugin
 	{
 		bool bSetDestination = false;
 
+		UGameFeaturesSubsystem& GameFeaturesSubsystem = UGameFeaturesSubsystem::Get();
 		if (TransitionPolicy::ExcludeDepedenciesFromBatchProcessing())
 		{
-			Dependency->ExcludeFromBatchProcessing();
+			GameFeaturesSubsystem.ExcludeFromBatchProcessing(Dependency);
 		}
 		
 		if (TransitionPolicy::ShouldWaitForDependencies())
 		{
-			bSetDestination = Dependency->SetDestination(TransitionPolicy::GetDependencyStateRange(),
+			bSetDestination = GameFeaturesSubsystem.SetDestination(Dependency, TransitionPolicy::GetDependencyStateRange(),
 				FGameFeatureStateTransitionComplete::CreateRaw(this, &FTransitionDependenciesGameFeaturePluginState::OnDependencyTransitionComplete));
 		}
 		else
 		{
-			bSetDestination = Dependency->SetDestination(TransitionPolicy::GetDependencyStateRange(), 
+			bSetDestination = GameFeaturesSubsystem.SetDestination(Dependency, TransitionPolicy::GetDependencyStateRange(), 
 				FGameFeatureStateTransitionComplete::CreateStatic(&FTransitionDependenciesGameFeaturePluginState::OnDependencyTransitionCompleteNoWait));
 			if (bSetDestination)
 			{
@@ -906,7 +928,7 @@ struct FTransitionDependenciesGameFeaturePluginState : public FGameFeaturePlugin
 
 		if (!bSetDestination)
 		{
-			const bool bCancelPending = Dependency->TryCancel(
+			const bool bCancelPending = GameFeaturesSubsystem.TryCancel(Dependency,
 				FGameFeatureStateTransitionCanceled::CreateRaw(this, &FTransitionDependenciesGameFeaturePluginState::OnDependencyTransitionCanceled));
 			if (!ensure(bCancelPending))
 			{
@@ -917,19 +939,20 @@ struct FTransitionDependenciesGameFeaturePluginState : public FGameFeaturePlugin
 
 	void OnDependencyTransitionCanceled(UGameFeaturePluginStateMachine* Dependency)
 	{
+		UGameFeaturesSubsystem& GameFeaturesSubsystem = UGameFeaturesSubsystem::Get();
 		// Special case for terminal state since it cannot be exited, we need to make a new machine
-		if (Dependency->GetCurrentState() == EGameFeaturePluginState::Terminal)
+		if (GameFeaturesSubsystem.GetCurrentState(Dependency) == EGameFeaturePluginState::Terminal)
 		{
 			// Inherit dep protocol options if possible
 			FGameFeatureProtocolOptions DepProtocolOptions;
-			EGameFeaturePluginProtocol DepProtocol = Dependency->GetPluginIdentifier().GetPluginProtocol();
+			EGameFeaturePluginProtocol DepProtocol = GameFeaturesSubsystem.GetPluginIdentifier(Dependency).GetPluginProtocol();
 			if (DepProtocol == EGameFeaturePluginProtocol::InstallBundle && StateProperties.ProtocolOptions.HasSubtype<FInstallBundlePluginProtocolOptions>())
 			{
 				DepProtocolOptions = StateProperties.RecycleProtocolOptions();
 			}
 
-			UGameFeaturePluginStateMachine* NewMachine = UGameFeaturesSubsystem::Get().FindOrCreateGameFeaturePluginStateMachine(Dependency->GetPluginURL(), DepProtocolOptions);
-			checkf(NewMachine != Dependency, TEXT("Game Feature Plugin %s should have already been removed from subsystem!"), *Dependency->GetPluginURL());
+			UGameFeaturePluginStateMachine* NewMachine = UGameFeaturesSubsystem::Get().FindOrCreateGameFeaturePluginStateMachine(GameFeaturesSubsystem.GetPluginURL(Dependency), DepProtocolOptions);
+			checkf(NewMachine != Dependency, TEXT("Game Feature Plugin %s should have already been removed from subsystem!"), *GameFeaturesSubsystem.GetPluginURL(Dependency));
 
 			const int32 Index = RemainingDependencies.IndexOfByPredicate([Dependency](const FDepResultPair& Pair)
 				{
@@ -940,8 +963,8 @@ struct FTransitionDependenciesGameFeaturePluginState : public FGameFeaturePlugin
 			FDepResultPair& FoundDep = RemainingDependencies[Index];
 			FoundDep.Key = NewMachine;
 
-			Dependency->RemovePendingTransitionCallback(this);
-			Dependency->RemovePendingCancelCallback(this);
+			GameFeaturesSubsystem.RemovePendingTransitionCallback(Dependency, this);
+			GameFeaturesSubsystem.RemovePendingCancelCallback(Dependency, this);
 
 			Dependency = NewMachine;
 		}
@@ -950,12 +973,12 @@ struct FTransitionDependenciesGameFeaturePluginState : public FGameFeaturePlugin
 		bool bSetDestination = false;
 		if (TransitionPolicy::ShouldWaitForDependencies())
 		{
-			bSetDestination = Dependency->SetDestination(TransitionPolicy::GetDependencyStateRange(),
+			bSetDestination = GameFeaturesSubsystem.SetDestination(Dependency, TransitionPolicy::GetDependencyStateRange(),
 				FGameFeatureStateTransitionComplete::CreateRaw(this, &FTransitionDependenciesGameFeaturePluginState::OnDependencyTransitionComplete));
 		}
 		else
 		{
-			bSetDestination = Dependency->SetDestination(TransitionPolicy::GetDependencyStateRange(), 
+			bSetDestination = GameFeaturesSubsystem.SetDestination(Dependency, TransitionPolicy::GetDependencyStateRange(), 
 				FGameFeatureStateTransitionComplete::CreateStatic(&FTransitionDependenciesGameFeaturePluginState::OnDependencyTransitionCompleteNoWait));
 			if (bSetDestination)
 			{
@@ -995,19 +1018,20 @@ struct FTransitionDependenciesGameFeaturePluginState : public FGameFeaturePlugin
 
 	static void OnDependencyTransitionCompleteNoWait(UGameFeaturePluginStateMachine* Dependency, const UE::GameFeatures::FResult& Result)
 	{
+		UGameFeaturesSubsystem& GameFeaturesSubsystem = UGameFeaturesSubsystem::Get();
 		if (Result.HasError())
 		{
 			if (Result.GetError() == UE::GameFeatures::CanceledResult.GetError())
 			{
 				UE_LOGFMT(LogGameFeatures, Warning, "Dependency {Dep} failed to transition because it was cancelled by another request {Error}",
-					("Dep", Dependency->GetPluginIdentifier().GetIdentifyingString()),
+					("Dep", GameFeaturesSubsystem.GetPluginIdentifier(Dependency).GetIdentifyingString()),
 					("Error", Result.GetError()));
 			}
 			else
 			{
 
 				UE_LOGFMT(LogGameFeatures, Error, "Dependency {Dep} failed to transition with error {Error}",
-					("Dep", Dependency->GetPluginIdentifier().GetIdentifyingString()),
+					("Dep", GameFeaturesSubsystem.GetPluginIdentifier(Dependency).GetIdentifyingString()),
 					("Error", Result.GetError()));
 			}
 		}
@@ -1017,13 +1041,14 @@ struct FTransitionDependenciesGameFeaturePluginState : public FGameFeaturePlugin
 	{
 		if (RemainingDependencies.Num() > 0)
 		{
+			UGameFeaturesSubsystem& GameFeaturesSubsystem = UGameFeaturesSubsystem::Get();
 			for (FDepResultPair& Pair : RemainingDependencies)
 			{
 				UGameFeaturePluginStateMachine* RemainingDependency = Pair.Key.Get();
 				if (RemainingDependency)
 				{
-					RemainingDependency->RemovePendingTransitionCallback(this);
-					RemainingDependency->RemovePendingCancelCallback(this);
+					GameFeaturesSubsystem.RemovePendingTransitionCallback(RemainingDependency, this);
+					GameFeaturesSubsystem.RemovePendingCancelCallback(RemainingDependency, this);
 				}
 			}
 
@@ -1045,7 +1070,7 @@ struct FTransitionDependenciesGameFeaturePluginState : public FGameFeaturePlugin
 
 struct FGameFeaturePluginState_Uninitialized : public FGameFeaturePluginState
 {
-	FGameFeaturePluginState_Uninitialized(FGameFeaturePluginStateMachineProperties& InStateProperties) : FGameFeaturePluginState(InStateProperties) {}
+	using FGameFeaturePluginState::FGameFeaturePluginState;
 
 	virtual void UpdateState(FGameFeaturePluginStateStatus& StateStatus) override
 	{
@@ -1055,8 +1080,7 @@ struct FGameFeaturePluginState_Uninitialized : public FGameFeaturePluginState
 
 struct FGameFeaturePluginState_Terminal : public FDestinationGameFeaturePluginState
 {
-	FGameFeaturePluginState_Terminal(FGameFeaturePluginStateMachineProperties& InStateProperties) : FDestinationGameFeaturePluginState(InStateProperties) {}
-
+	using FDestinationGameFeaturePluginState::FDestinationGameFeaturePluginState;
 	bool bEnteredTerminalState = false;
 
 	virtual UE::GameFeatures::FResult TryUpdateProtocolOptions(const FGameFeatureProtocolOptions& NewOptions) override
@@ -1076,7 +1100,7 @@ struct FGameFeaturePluginState_Terminal : public FDestinationGameFeaturePluginSt
 
 struct FGameFeaturePluginState_UnknownStatus : public FDestinationGameFeaturePluginState
 {
-	FGameFeaturePluginState_UnknownStatus(FGameFeaturePluginStateMachineProperties& InStateProperties) : FDestinationGameFeaturePluginState(InStateProperties) {}
+	using FDestinationGameFeaturePluginState::FDestinationGameFeaturePluginState;
 
 	virtual void UpdateState(FGameFeaturePluginStateStatus& StateStatus) override
 	{
@@ -1095,7 +1119,7 @@ struct FGameFeaturePluginState_UnknownStatus : public FDestinationGameFeaturePlu
 
 struct FGameFeaturePluginState_CheckingStatus : public FGameFeaturePluginState
 {
-	FGameFeaturePluginState_CheckingStatus(FGameFeaturePluginStateMachineProperties& InStateProperties) : FGameFeaturePluginState(InStateProperties) {}
+	using FGameFeaturePluginState::FGameFeaturePluginState;
 
 	bool bParsedURL = false;
 	bool bIsAvailable = false;
@@ -1188,7 +1212,7 @@ struct FGameFeaturePluginState_CheckingStatus : public FGameFeaturePluginState
 
 struct FGameFeaturePluginState_ErrorCheckingStatus : public FErrorGameFeaturePluginState
 {
-	FGameFeaturePluginState_ErrorCheckingStatus(FGameFeaturePluginStateMachineProperties& InStateProperties) : FErrorGameFeaturePluginState(InStateProperties) {}
+	using FErrorGameFeaturePluginState::FErrorGameFeaturePluginState;
 
 	virtual void UpdateState(FGameFeaturePluginStateStatus& StateStatus) override
 	{
@@ -1205,7 +1229,7 @@ struct FGameFeaturePluginState_ErrorCheckingStatus : public FErrorGameFeaturePlu
 
 struct FGameFeaturePluginState_ErrorUnavailable : public FErrorGameFeaturePluginState
 {
-	FGameFeaturePluginState_ErrorUnavailable(FGameFeaturePluginStateMachineProperties& InStateProperties) : FErrorGameFeaturePluginState(InStateProperties) {}
+	using FErrorGameFeaturePluginState::FErrorGameFeaturePluginState;
 
 	virtual void UpdateState(FGameFeaturePluginStateStatus& StateStatus) override
 	{
@@ -1222,7 +1246,7 @@ struct FGameFeaturePluginState_ErrorUnavailable : public FErrorGameFeaturePlugin
 
 struct FGameFeaturePluginState_StatusKnown : public FDestinationGameFeaturePluginState
 {
-	FGameFeaturePluginState_StatusKnown(FGameFeaturePluginStateMachineProperties& InStateProperties) : FDestinationGameFeaturePluginState(InStateProperties) {}
+	using FDestinationGameFeaturePluginState::FDestinationGameFeaturePluginState;
 
 	virtual void UpdateState(FGameFeaturePluginStateStatus& StateStatus) override
 	{
@@ -1253,7 +1277,7 @@ struct FGameFeaturePluginState_StatusKnown : public FDestinationGameFeaturePlugi
 
 struct FGameFeaturePluginState_ErrorManagingData : public FErrorGameFeaturePluginState
 {
-	FGameFeaturePluginState_ErrorManagingData(FGameFeaturePluginStateMachineProperties& InStateProperties) : FErrorGameFeaturePluginState(InStateProperties) {}
+	using FErrorGameFeaturePluginState::FErrorGameFeaturePluginState;
 
 	virtual void UpdateState(FGameFeaturePluginStateStatus& StateStatus) override
 	{
@@ -1270,7 +1294,7 @@ struct FGameFeaturePluginState_ErrorManagingData : public FErrorGameFeaturePlugi
 
 struct FGameFeaturePluginState_ErrorUninstalling : public FErrorGameFeaturePluginState
 {
-	FGameFeaturePluginState_ErrorUninstalling(FGameFeaturePluginStateMachineProperties& InStateProperties) : FErrorGameFeaturePluginState(InStateProperties) {}
+	using FErrorGameFeaturePluginState::FErrorGameFeaturePluginState;
 
 	virtual void UpdateState(FGameFeaturePluginStateStatus& StateStatus) override
 	{
@@ -1295,15 +1319,12 @@ struct FGameFeaturePluginState_ErrorUninstalling : public FErrorGameFeaturePlugi
 //Base class for our states that want to request to release data from InstallBundleManager
 struct FBaseDataReleaseGameFeaturePluginState : public FGameFeaturePluginState
 {
-	FBaseDataReleaseGameFeaturePluginState(FGameFeaturePluginStateMachineProperties& InStateProperties) 
-		: FGameFeaturePluginState(InStateProperties) 
-		, Result(MakeValue())
-	{}
+	using FGameFeaturePluginState::FGameFeaturePluginState;
 
 	virtual ~FBaseDataReleaseGameFeaturePluginState()
 	{}
 
-	UE::GameFeatures::FResult Result;
+	UE::GameFeatures::FResult Result = MakeValue();
 	bool bWasDeleted = false;
 	TArray<FName> PendingBundles;
 
@@ -1463,7 +1484,7 @@ struct FBaseDataReleaseGameFeaturePluginState : public FGameFeaturePluginState
 
 struct FGameFeaturePluginState_Uninstalled : public FDestinationGameFeaturePluginState
 {
-	FGameFeaturePluginState_Uninstalled(FGameFeaturePluginStateMachineProperties& InStateProperties) : FDestinationGameFeaturePluginState(InStateProperties) {}
+	using FDestinationGameFeaturePluginState::FDestinationGameFeaturePluginState;
 
 	virtual void UpdateState(FGameFeaturePluginStateStatus& StateStatus) override
 	{
@@ -1480,9 +1501,7 @@ struct FGameFeaturePluginState_Uninstalled : public FDestinationGameFeaturePlugi
 
 struct FGameFeaturePluginState_Uninstalling : public FBaseDataReleaseGameFeaturePluginState
 {
-	FGameFeaturePluginState_Uninstalling(FGameFeaturePluginStateMachineProperties& InStateProperties)
-		: FBaseDataReleaseGameFeaturePluginState(InStateProperties)
-	{}
+	using FBaseDataReleaseGameFeaturePluginState::FBaseDataReleaseGameFeaturePluginState;
 
 	virtual EGameFeaturePluginState GetSuccessTransitionState() const override
 	{
@@ -1532,9 +1551,7 @@ struct FGameFeaturePluginState_Uninstalling : public FBaseDataReleaseGameFeature
 
 struct FGameFeaturePluginState_Releasing : public FBaseDataReleaseGameFeaturePluginState
 {
-	FGameFeaturePluginState_Releasing(FGameFeaturePluginStateMachineProperties& InStateProperties)
-		: FBaseDataReleaseGameFeaturePluginState(InStateProperties)
-	{}
+	using FBaseDataReleaseGameFeaturePluginState::FBaseDataReleaseGameFeaturePluginState;
 
 	virtual void BeginState() override
 	{
@@ -1559,17 +1576,14 @@ struct FGameFeaturePluginState_Releasing : public FBaseDataReleaseGameFeaturePlu
 
 struct FGameFeaturePluginState_Downloading : public FGameFeaturePluginState
 {
-	FGameFeaturePluginState_Downloading(FGameFeaturePluginStateMachineProperties& InStateProperties)
-		: FGameFeaturePluginState(InStateProperties)
-		, Result(MakeValue())
-	{}
+	using FGameFeaturePluginState::FGameFeaturePluginState;
 
 	virtual ~FGameFeaturePluginState_Downloading()
 	{
 		Cleanup();
 	}
 
-	UE::GameFeatures::FResult Result;
+	UE::GameFeatures::FResult Result = MakeValue();
 	bool bSuppressResultErrorLog = false;
 	bool bPluginDownloaded = false;
 
@@ -2063,7 +2077,7 @@ struct FGameFeaturePluginState_Downloading : public FGameFeaturePluginState
 
 struct FGameFeaturePluginState_Installed : public FDestinationGameFeaturePluginState
 {
-	FGameFeaturePluginState_Installed(FGameFeaturePluginStateMachineProperties& InStateProperties) : FDestinationGameFeaturePluginState(InStateProperties) {}
+	using FDestinationGameFeaturePluginState::FDestinationGameFeaturePluginState;
 
 	virtual void UpdateState(FGameFeaturePluginStateStatus& StateStatus) override
 	{
@@ -2080,7 +2094,7 @@ struct FGameFeaturePluginState_Installed : public FDestinationGameFeaturePluginS
 
 struct FGameFeaturePluginState_ErrorMounting : public FErrorGameFeaturePluginState
 {
-	FGameFeaturePluginState_ErrorMounting(FGameFeaturePluginStateMachineProperties& InStateProperties) : FErrorGameFeaturePluginState(InStateProperties) {}
+	using FErrorGameFeaturePluginState::FErrorGameFeaturePluginState;
 
 	virtual void UpdateState(FGameFeaturePluginStateStatus& StateStatus) override
 	{
@@ -2097,7 +2111,7 @@ struct FGameFeaturePluginState_ErrorMounting : public FErrorGameFeaturePluginSta
 
 struct FGameFeaturePluginState_ErrorWaitingForDependencies : public FErrorGameFeaturePluginState
 {
-	FGameFeaturePluginState_ErrorWaitingForDependencies(FGameFeaturePluginStateMachineProperties& InStateProperties) : FErrorGameFeaturePluginState(InStateProperties) {}
+	using FErrorGameFeaturePluginState::FErrorGameFeaturePluginState;
 
 	virtual void UpdateState(FGameFeaturePluginStateStatus& StateStatus) override
 	{
@@ -2115,7 +2129,7 @@ struct FGameFeaturePluginState_ErrorWaitingForDependencies : public FErrorGameFe
 
 struct FGameFeaturePluginState_ErrorRegistering : public FErrorGameFeaturePluginState
 {
-	FGameFeaturePluginState_ErrorRegistering(FGameFeaturePluginStateMachineProperties& InStateProperties) : FErrorGameFeaturePluginState(InStateProperties) {}
+	using FErrorGameFeaturePluginState::FErrorGameFeaturePluginState;
 
 	virtual void UpdateState(FGameFeaturePluginStateStatus& StateStatus) override
 	{
@@ -2132,7 +2146,7 @@ struct FGameFeaturePluginState_ErrorRegistering : public FErrorGameFeaturePlugin
 
 struct FGameFeaturePluginState_Unmounting : public FGameFeaturePluginState
 {
-	FGameFeaturePluginState_Unmounting(FGameFeaturePluginStateMachineProperties& InStateProperties) : FGameFeaturePluginState(InStateProperties) {}
+	using FGameFeaturePluginState::FGameFeaturePluginState;
 
 	UE::GameFeatures::FResult Result = MakeValue();
 	TArray<FName> PendingBundles;
@@ -2305,10 +2319,7 @@ struct FGameFeaturePluginState_Unmounting : public FGameFeaturePluginState
 
 struct FGameFeaturePluginState_Mounting : public FGameFeaturePluginState
 {
-	FGameFeaturePluginState_Mounting(FGameFeaturePluginStateMachineProperties& InStateProperties)
-		: FGameFeaturePluginState(InStateProperties)
-		, Result(MakeValue())
-	{}
+	using FGameFeaturePluginState::FGameFeaturePluginState;
 
 	enum class ESubState : uint8
 	{
@@ -2322,7 +2333,7 @@ struct FGameFeaturePluginState_Mounting : public FGameFeaturePluginState
 	int32 NumExpectedPostMountPausers = 0;
 	TArray<FName> PendingBundles;
 	FDelegateHandle PakFileMountedDelegateHandle;
-	UE::GameFeatures::FResult Result;
+	UE::GameFeatures::FResult Result = MakeValue();
 	ESubState StartedSubStates = ESubState::None;
 	ESubState CompletedSubStates = ESubState::None;
 	bool bCheckedRealtimeMode = false;
@@ -2839,17 +2850,13 @@ struct FWaitingForDependenciesTransitionPolicy
 
 struct FGameFeaturePluginState_WaitingForDependencies : public FTransitionDependenciesGameFeaturePluginState<FWaitingForDependenciesTransitionPolicy>
 {
-	FGameFeaturePluginState_WaitingForDependencies(FGameFeaturePluginStateMachineProperties& InStateProperties)
-		: FTransitionDependenciesGameFeaturePluginState(InStateProperties)
-	{
-	}
+	using Super = FTransitionDependenciesGameFeaturePluginState<FWaitingForDependenciesTransitionPolicy>;
+	using Super::Super;
 };
 
 struct FGameFeaturePluginState_AssetDependencyStreamOut : public FGameFeaturePluginState
 {
-	FGameFeaturePluginState_AssetDependencyStreamOut(FGameFeaturePluginStateMachineProperties& InStateProperties)
-		: FGameFeaturePluginState(InStateProperties)
-	{}
+	using FGameFeaturePluginState::FGameFeaturePluginState;
 
 	virtual void BeginState() override
 	{
@@ -2872,9 +2879,7 @@ struct FGameFeaturePluginState_AssetDependencyStreamOut : public FGameFeaturePlu
 
 struct FGameFeaturePluginState_ErrorAssetDependencyStreaming : public FErrorGameFeaturePluginState
 {
-	FGameFeaturePluginState_ErrorAssetDependencyStreaming(FGameFeaturePluginStateMachineProperties& InStateProperties) 
-		: FErrorGameFeaturePluginState(InStateProperties) 
-	{}
+	using FErrorGameFeaturePluginState::FErrorGameFeaturePluginState;
 
 	virtual void UpdateState(FGameFeaturePluginStateStatus& StateStatus) override
 	{
@@ -2891,9 +2896,7 @@ struct FGameFeaturePluginState_ErrorAssetDependencyStreaming : public FErrorGame
 
 struct FGameFeaturePluginState_AssetDependencyStreaming : public FGameFeaturePluginState
 {
-	FGameFeaturePluginState_AssetDependencyStreaming(FGameFeaturePluginStateMachineProperties& InStateProperties)
-		: FGameFeaturePluginState(InStateProperties)
-	{}
+	using FGameFeaturePluginState::FGameFeaturePluginState;
 
 	virtual ~FGameFeaturePluginState_AssetDependencyStreaming()
 	{
@@ -3096,7 +3099,7 @@ struct FGameFeaturePluginState_AssetDependencyStreaming : public FGameFeaturePlu
 
 struct FGameFeaturePluginState_Unregistering : public FGameFeaturePluginState
 {
-	FGameFeaturePluginState_Unregistering(FGameFeaturePluginStateMachineProperties& InStateProperties) : FGameFeaturePluginState(InStateProperties) {}
+	using FGameFeaturePluginState::FGameFeaturePluginState;
 
 	bool bHasUnloaded = false;
 #if WITH_EDITOR
@@ -3168,6 +3171,8 @@ struct FGameFeaturePluginState_Unregistering : public FGameFeaturePluginState
 
 struct FGameFeaturePluginState_Registering : public FGameFeaturePluginState
 {
+	using FGameFeaturePluginState::FGameFeaturePluginState;
+
 	enum class ELoadGFDState : uint8
 	{
 		Pending = 0,
@@ -3181,14 +3186,13 @@ struct FGameFeaturePluginState_Registering : public FGameFeaturePluginState
 	bool bCheckedRealtimeMode = false;
 	TArray<FString, TInlineAllocator<2>> GameFeatureDataPaths;
 
-	FGameFeaturePluginState_Registering(FGameFeaturePluginStateMachineProperties& InStateProperties) : FGameFeaturePluginState(InStateProperties) {}
-	
 	void TryAsyncLoadGameFeatureData(int32 Attempt = 0)
 	{
 		check(LoadGFDState == ELoadGFDState::Pending);
 
 		if (!GameFeatureDataPaths.IsValidIndex(Attempt))
 		{
+			GFP_STATE_LOG(VeryVerbose, "Ran out of game feature paths to try");
 			LoadGFDState = ELoadGFDState::Failed;
 			UpdateStateMachineDeferred();
 			return;
@@ -3196,9 +3200,11 @@ struct FGameFeaturePluginState_Registering : public FGameFeaturePluginState
 
 		bool bIsLoading = false;
 
+		GFP_STATE_LOG(VeryVerbose, "Trying to load game feature data {GameFeatureData}", GameFeatureDataPaths[Attempt]);
 		GameFeatureDataHandle = UGameFeaturesSubsystem::LoadGameFeatureData(GameFeatureDataPaths[Attempt], true /*bStartStalled*/);
 		if (GameFeatureDataHandle && GameFeatureDataHandle->IsLoadingInProgress())
 		{
+			GFP_STATE_LOG(VeryVerbose, "Waiting for load of game feature data {GameFeatureData}", GameFeatureDataPaths[Attempt]);
 			// CurrentMachine owns `this`, so use it as a lifetime check for the `this` captured in the lambdas below, as they may be called after `this` is destroyed
 			UGameFeaturePluginStateMachine* CurrentMachine = UGameFeaturesSubsystem::Get().FindGameFeaturePluginStateMachine(StateProperties.PluginIdentifier);
 
@@ -3206,12 +3212,13 @@ struct FGameFeaturePluginState_Registering : public FGameFeaturePluginState
 			{
 				if (GameFeatureDataHandle != Handle)
 				{
+					GFP_STATE_LOG(VeryVerbose, "Ignoring cancel of game feature data load as we are no longer in a relevant state");
 					// We're no longer in a state where we want to process this callback (maybe we already went through EndState)
 					return;
 				}
 
 				const FStringView ShortUrl = StateProperties.PluginIdentifier.GetIdentifyingString();
-				UE_LOG(LogGameFeatures, Error, TEXT("Game Feature Data loading was cancelled for URL %.*s"), ShortUrl.Len(), ShortUrl.GetData());
+				GFP_STATE_LOG(Error, "Game Feature Data loading was cancelled for URL {ShortURL}", ShortUrl);
 
 				LoadGFDState = ELoadGFDState::Cancelled;
 				UpdateStateMachineDeferred();
@@ -3221,6 +3228,7 @@ struct FGameFeaturePluginState_Registering : public FGameFeaturePluginState
 			{
 				if (GameFeatureDataHandle != Handle)
 				{
+					GFP_STATE_LOG(VeryVerbose, "Ignoring completion of game feature data load as we are no longer in a relevant state");
 					// We're no longer in a state where we want to process this callback (maybe we already went through EndState)
 					return;
 				}
@@ -3234,6 +3242,7 @@ struct FGameFeaturePluginState_Registering : public FGameFeaturePluginState
 
 				if (StateProperties.GameFeatureData)
 				{
+					GFP_STATE_LOG(VeryVerbose, "Loaded GameFeatureData {GameFeatureData}", FSoftObjectPath(StateProperties.GameFeatureData));
 					LoadGFDState = ELoadGFDState::Success;
 				}
 				else
@@ -3326,6 +3335,7 @@ struct FGameFeaturePluginState_Registering : public FGameFeaturePluginState
 		// If the plugin contains content then load the GameFeatureData otherwise procedurally create one that is transient.
 		if (!Plugin->GetDescriptor().bCanContainContent)
 		{
+			GFP_STATE_LOG(VeryVerbose, "Creating transient GameFeatureData for non-content plugin");
 			StateProperties.GameFeatureData = CreateTransientGameFeatureData();
 			LoadGFDState = ELoadGFDState::Success;
 			return;
@@ -3396,6 +3406,7 @@ struct FGameFeaturePluginState_Registering : public FGameFeaturePluginState
 			
 			for (const FString& Path : GameFeatureDataPaths)
 			{
+				GFP_STATE_LOG(VeryVerbose, "Synchronously loading GameFeatureData {Path}", Path);
 				GameFeatureDataHandle = UGameFeaturesSubsystem::LoadGameFeatureData(Path);
 				if (GameFeatureDataHandle)
 				{
@@ -3414,6 +3425,7 @@ struct FGameFeaturePluginState_Registering : public FGameFeaturePluginState
 			}
 			else
 			{
+				GFP_STATE_LOG(VeryVerbose, "Failed to synchronously load GameFeatureData");
 				// This will set LoadGFDState
 				HandleMissingGameFeatureData();
 			}
@@ -3429,7 +3441,8 @@ struct FGameFeaturePluginState_Registering : public FGameFeaturePluginState
 	virtual void UpdateState(FGameFeaturePluginStateStatus& StateStatus) override
 	{
 		TRACE_CPUPROFILER_EVENT_SCOPE(GFP_Registering_Update);
-
+		UGameFeaturesSubsystem& GameFeaturesSubsystem = UGameFeaturesSubsystem::Get();
+		
 		if (!bCheckedRealtimeMode)
 		{
 			bCheckedRealtimeMode = true;
@@ -3463,9 +3476,9 @@ struct FGameFeaturePluginState_Registering : public FGameFeaturePluginState
 			StateStatus.SetTransition(EGameFeaturePluginState::Registered);
 
 			check(StateProperties.AddedPrimaryAssetTypes.Num() == 0);
-			UGameFeaturesSubsystem::Get().AddGameFeatureToAssetManager(StateProperties.GameFeatureData, StateProperties.PluginName, StateProperties.AddedPrimaryAssetTypes);
+			GameFeaturesSubsystem.AddGameFeatureToAssetManager(StateProperties.GameFeatureData, StateProperties.PluginName, StateProperties.AddedPrimaryAssetTypes);
 
-			UGameFeaturesSubsystem::Get().OnGameFeatureRegistering(StateProperties.GameFeatureData, StateProperties.PluginName, StateProperties.PluginIdentifier);
+			GameFeaturesSubsystem.OnGameFeatureRegistering(StateProperties.GameFeatureData, StateProperties.PluginName, StateProperties.PluginIdentifier);
 		}
 		else
 		{
@@ -3477,7 +3490,7 @@ struct FGameFeaturePluginState_Registering : public FGameFeaturePluginState
 			{
 				const FStringView ShortUrl = StateProperties.PluginIdentifier.GetIdentifyingString();
 				UE_LOG(LogGameFeatures, Error, TEXT("Setting %.*s to be in unrecoverable error as GameFeatureData is missing"), ShortUrl.Len(), ShortUrl.GetData());
-				CurrentMachine->SetUnrecoverableError();
+				GameFeaturesSubsystem.SetUnrecoverableError(CurrentMachine);
 			}
 		}
 	}
@@ -3485,7 +3498,7 @@ struct FGameFeaturePluginState_Registering : public FGameFeaturePluginState
 
 struct FGameFeaturePluginState_Registered : public FDestinationGameFeaturePluginState
 {
-	FGameFeaturePluginState_Registered(FGameFeaturePluginStateMachineProperties& InStateProperties) : FDestinationGameFeaturePluginState(InStateProperties) {}
+	using FDestinationGameFeaturePluginState::FDestinationGameFeaturePluginState;
 
 	virtual void UpdateState(FGameFeaturePluginStateStatus& StateStatus) override
 	{
@@ -3502,7 +3515,7 @@ struct FGameFeaturePluginState_Registered : public FDestinationGameFeaturePlugin
 
 struct FGameFeaturePluginState_ErrorLoading : public FErrorGameFeaturePluginState
 {
-	FGameFeaturePluginState_ErrorLoading(FGameFeaturePluginStateMachineProperties& InStateProperties) : FErrorGameFeaturePluginState(InStateProperties) {}
+	using FErrorGameFeaturePluginState::FErrorGameFeaturePluginState;
 
 	virtual void UpdateState(FGameFeaturePluginStateStatus& StateStatus) override
 	{
@@ -3519,7 +3532,7 @@ struct FGameFeaturePluginState_ErrorLoading : public FErrorGameFeaturePluginStat
 
 struct FGameFeaturePluginState_Unloading : public FGameFeaturePluginState
 {
-	FGameFeaturePluginState_Unloading(FGameFeaturePluginStateMachineProperties& InStateProperties) : FGameFeaturePluginState(InStateProperties) {}
+	using FGameFeaturePluginState::FGameFeaturePluginState;
 
 	virtual void BeginState() override
 	{
@@ -3566,7 +3579,7 @@ struct FGameFeaturePluginState_Unloading : public FGameFeaturePluginState
 
 struct FGameFeaturePluginState_Loading : public FGameFeaturePluginState
 {
-	FGameFeaturePluginState_Loading(FGameFeaturePluginStateMachineProperties& InStateProperties) : FGameFeaturePluginState(InStateProperties) {}
+	using FGameFeaturePluginState::FGameFeaturePluginState;
 
 	TSharedPtr<FStreamableHandle> BundleHandle;
 
@@ -3678,7 +3691,7 @@ struct FGameFeaturePluginState_Loading : public FGameFeaturePluginState
 
 struct FGameFeaturePluginState_Loaded : public FDestinationGameFeaturePluginState
 {
-	FGameFeaturePluginState_Loaded(FGameFeaturePluginStateMachineProperties& InStateProperties) : FDestinationGameFeaturePluginState(InStateProperties) {}
+	using FDestinationGameFeaturePluginState::FDestinationGameFeaturePluginState;
 
 	virtual void UpdateState(FGameFeaturePluginStateStatus& StateStatus) override
 	{
@@ -3695,7 +3708,7 @@ struct FGameFeaturePluginState_Loaded : public FDestinationGameFeaturePluginStat
 
 struct FGameFeaturePluginState_ErrorDeactivatingDependencies : public FErrorGameFeaturePluginState
 {
-	FGameFeaturePluginState_ErrorDeactivatingDependencies(FGameFeaturePluginStateMachineProperties& InStateProperties) : FErrorGameFeaturePluginState(InStateProperties) {}
+	using FErrorGameFeaturePluginState::FErrorGameFeaturePluginState;
 
 	virtual void UpdateState(FGameFeaturePluginStateStatus& StateStatus) override
 	{
@@ -3748,15 +3761,13 @@ struct FDeactivatingDependenciesTransitionPolicy
 
 struct FGameFeaturePluginState_DeactivatingDependencies : public FTransitionDependenciesGameFeaturePluginState<FDeactivatingDependenciesTransitionPolicy>
 {
-	FGameFeaturePluginState_DeactivatingDependencies(FGameFeaturePluginStateMachineProperties& InStateProperties)
-		: FTransitionDependenciesGameFeaturePluginState(InStateProperties)
-	{
-	}
+	using Super = FTransitionDependenciesGameFeaturePluginState<FDeactivatingDependenciesTransitionPolicy>;
+	using Super::Super;
 };
 
 struct FGameFeaturePluginState_Deactivating : public FGameFeaturePluginState
 {
-	FGameFeaturePluginState_Deactivating(FGameFeaturePluginStateMachineProperties& InStateProperties) : FGameFeaturePluginState(InStateProperties) {}
+	using FGameFeaturePluginState::FGameFeaturePluginState;
 
 	int32 NumObservedPausers = 0;
 	int32 NumExpectedPausers = 0;
@@ -3855,7 +3866,7 @@ struct FGameFeaturePluginState_Deactivating : public FGameFeaturePluginState
 
 struct FGameFeaturePluginState_ErrorActivatingDependencies : public FErrorGameFeaturePluginState
 {
-	FGameFeaturePluginState_ErrorActivatingDependencies(FGameFeaturePluginStateMachineProperties& InStateProperties) : FErrorGameFeaturePluginState(InStateProperties) {}
+	using FErrorGameFeaturePluginState::FErrorGameFeaturePluginState;
 
 	virtual void UpdateState(FGameFeaturePluginStateStatus& StateStatus) override
 	{
@@ -3909,15 +3920,13 @@ struct FActivatingDependenciesTransitionPolicy
 
 struct FGameFeaturePluginState_ActivatingDependencies : public FTransitionDependenciesGameFeaturePluginState<FActivatingDependenciesTransitionPolicy>
 {
-	FGameFeaturePluginState_ActivatingDependencies(FGameFeaturePluginStateMachineProperties& InStateProperties)
-		: FTransitionDependenciesGameFeaturePluginState(InStateProperties)
-	{
-	}
+	using Super = FTransitionDependenciesGameFeaturePluginState<FActivatingDependenciesTransitionPolicy>;
+	using Super::Super;
 };
 
 struct FGameFeaturePluginState_Activating : public FGameFeaturePluginState
 {
-	FGameFeaturePluginState_Activating(FGameFeaturePluginStateMachineProperties& InStateProperties) : FGameFeaturePluginState(InStateProperties) {}
+	using FGameFeaturePluginState::FGameFeaturePluginState;
 
 	bool CanBatchProcess() const override
 	{
@@ -3975,7 +3984,7 @@ struct FGameFeaturePluginState_Activating : public FGameFeaturePluginState
 
 		for (const UGameFeaturePluginStateMachine* GFPSM : GFPs)
 		{
-			PluginInstalledFilenames.Emplace(GFPSM->GetProperties().PluginInstalledFilename);
+			PluginInstalledFilenames.Emplace(GFPSubSys.GetProperties(GFPSM).PluginInstalledFilename);
 		}
 
 		{
@@ -3987,7 +3996,7 @@ struct FGameFeaturePluginState_Activating : public FGameFeaturePluginState
 
 struct FGameFeaturePluginState_Active : public FDestinationGameFeaturePluginState
 {
-	FGameFeaturePluginState_Active(FGameFeaturePluginStateMachineProperties& InStateProperties) : FDestinationGameFeaturePluginState(InStateProperties) {}
+	using FDestinationGameFeaturePluginState::FDestinationGameFeaturePluginState;
 
 	virtual void BeginState() override
 	{
@@ -4078,34 +4087,54 @@ UGameFeaturePluginStateMachine::UGameFeaturePluginStateMachine(const FObjectInit
 
 }
 
-void UGameFeaturePluginStateMachine::InitStateMachine(FGameFeaturePluginIdentifier InPluginIdentifier, const FGameFeatureProtocolOptions& InProtocolOptions)
+void UGameFeaturesSubsystem::InitStateMachine(UGameFeaturePluginStateMachine* This, FGameFeaturePluginIdentifier InPluginIdentifier, const FGameFeatureProtocolOptions& InProtocolOptions) 
 {
 	LLM_SCOPE_BYTAG(GFP);
 
-	check(GetCurrentState() == EGameFeaturePluginState::Uninitialized);
-	CurrentStateInfo.State = EGameFeaturePluginState::UnknownStatus;
-	StateProperties = FGameFeaturePluginStateMachineProperties(
+	check(GetCurrentState(This) == EGameFeaturePluginState::Uninitialized);
+	This->CurrentStateInfo.State = EGameFeaturePluginState::UnknownStatus;
+	This->StateProperties = FGameFeaturePluginStateMachineProperties(
 		MoveTemp(InPluginIdentifier),
-		FGameFeaturePluginStateRange(CurrentStateInfo.State),
-		FGameFeaturePluginRequestUpdateStateMachine::CreateUObject(this, &ThisClass::UpdateStateMachine),
-		FGameFeatureStateProgressUpdate::CreateUObject(this, &ThisClass::UpdateCurrentStateProgress));
+		FGameFeaturePluginStateRange(This->CurrentStateInfo.State),
+		FGameFeaturePluginRequestUpdateStateMachine::CreateWeakLambda(this, [this, InPluginIdentifier, WeakThis=MakeWeakObjectPtr(This)]()
+		{
+			if (UGameFeaturePluginStateMachine* This = WeakThis.Get())
+			{
+				UpdateStateMachine(This);
+			}
+			else
+			{
+				UE_LOGFMT(LogGameFeatures, Warning, "State machine pointer for {Plugin} became unexpectedly null.", InPluginIdentifier.GetPluginName());
+			}
+		}),
+		FGameFeatureStateProgressUpdate::CreateWeakLambda(this, [this, InPluginIdentifier, WeakThis=MakeWeakObjectPtr(This)](float Progress)
+		{
+			if (UGameFeaturePluginStateMachine* This = WeakThis.Get())
+			{
+				UpdateCurrentStateProgress(This, Progress);
+			}
+			else
+			{
+				UE_LOGFMT(LogGameFeatures, Warning, "State machine pointer for {Plugin} became unexpectedly null.", InPluginIdentifier.GetPluginName());
+			}
+		}));
 
-	StateProperties.ProtocolOptions = InProtocolOptions;
+	This->StateProperties.ProtocolOptions = InProtocolOptions;
 
-#define GAME_FEATURE_PLUGIN_STATE_MAKE_STATE(inEnum, inText) AllStates[EGameFeaturePluginState::inEnum] = MakeUnique<FGameFeaturePluginState_##inEnum>(StateProperties);
+#define GAME_FEATURE_PLUGIN_STATE_MAKE_STATE(inEnum, inText) This->AllStates[EGameFeaturePluginState::inEnum] = MakeUnique<FGameFeaturePluginState_##inEnum>(EGameFeaturePluginState::inEnum, This->StateProperties);
 	GAME_FEATURE_PLUGIN_STATE_LIST(GAME_FEATURE_PLUGIN_STATE_MAKE_STATE)
 #undef GAME_FEATURE_PLUGIN_STATE_MAKE_STATE
 
-	CheckAddBatchingRequestForCurrentState();
-	AllStates[CurrentStateInfo.State]->BeginState();
+	CheckAddBatchingRequestForCurrentState(This);
+	This->AllStates[This->CurrentStateInfo.State]->BeginState();
 }
 
-bool UGameFeaturePluginStateMachine::SetDestination(FGameFeaturePluginStateRange InDestination, FGameFeatureStateTransitionComplete OnFeatureStateTransitionComplete, FDelegateHandle* OutCallbackHandle /*= nullptr*/)
+bool UGameFeaturesSubsystem::SetDestination(UGameFeaturePluginStateMachine* This, FGameFeaturePluginStateRange InDestination, FGameFeatureStateTransitionComplete OnFeatureStateTransitionComplete, FDelegateHandle* OutCallbackHandle /*= nullptr*/)
 {
 	LLM_SCOPE_BYTAG(GFP);
 
-	check(IsValidDestinationState(InDestination.MinState));
-	check(IsValidDestinationState(InDestination.MaxState));
+	check(IsValidDestinationState(This, InDestination.MinState));
+	check(IsValidDestinationState(This, InDestination.MaxState));
 
 	bool bDestinationSet = false;
 	bool bDestinationChanged = false;
@@ -4114,11 +4143,11 @@ bool UGameFeaturePluginStateMachine::SetDestination(FGameFeaturePluginStateRange
 	{
 		// Invalid range
 	}
-	else if (CurrentStateInfo.State == EGameFeaturePluginState::Terminal && !InDestination.Contains(EGameFeaturePluginState::Terminal))
+	else if (This->CurrentStateInfo.State == EGameFeaturePluginState::Terminal && !InDestination.Contains(EGameFeaturePluginState::Terminal))
 	{
 		// Can't tranistion away from terminal state
 	}
-	else if (!IsRunning())
+	else if (!IsRunning(This))
 	{
 		// Not running so any new range is acceptable
 
@@ -4127,26 +4156,26 @@ bool UGameFeaturePluginStateMachine::SetDestination(FGameFeaturePluginStateRange
 			OutCallbackHandle->Reset();
 		}
 
-		FDestinationGameFeaturePluginState* CurrState = AllStates[CurrentStateInfo.State]->AsDestinationState();
+		FDestinationGameFeaturePluginState* CurrState = This->AllStates[This->CurrentStateInfo.State]->AsDestinationState();
 
-		if (InDestination.Contains(CurrentStateInfo.State))
+		if (InDestination.Contains(This->CurrentStateInfo.State))
 		{
-			OnFeatureStateTransitionComplete.ExecuteIfBound(this, MakeValue());
+			OnFeatureStateTransitionComplete.ExecuteIfBound(This, MakeValue());
 		}
 		else
 		{
-			if (CurrentStateInfo.State < InDestination)
+			if (This->CurrentStateInfo.State < InDestination)
 			{
-				FDestinationGameFeaturePluginState* MinDestState = AllStates[InDestination.MinState]->AsDestinationState();
+				FDestinationGameFeaturePluginState* MinDestState = This->AllStates[InDestination.MinState]->AsDestinationState();
 				FDelegateHandle CallbackHandle = MinDestState->OnDestinationStateReached.Add(MoveTemp(OnFeatureStateTransitionComplete));
 				if (OutCallbackHandle)
 				{
 					*OutCallbackHandle = CallbackHandle;
 				}
 			}
-			else if (CurrentStateInfo.State > InDestination)
+			else if (This->CurrentStateInfo.State > InDestination)
 			{
-				FDestinationGameFeaturePluginState* MaxDestState = AllStates[InDestination.MaxState]->AsDestinationState();
+				FDestinationGameFeaturePluginState* MaxDestState = This->AllStates[InDestination.MaxState]->AsDestinationState();
 				FDelegateHandle CallbackHandle = MaxDestState->OnDestinationStateReached.Add(MoveTemp(OnFeatureStateTransitionComplete));
 				if (OutCallbackHandle)
 				{
@@ -4154,14 +4183,14 @@ bool UGameFeaturePluginStateMachine::SetDestination(FGameFeaturePluginStateRange
 				}
 			}
 
-			StateProperties.Destination = InDestination;
-			UpdateStateMachine();
+			This->StateProperties.Destination = InDestination;
+			UpdateStateMachine(This);
 			bDestinationChanged = true;
 		}
 
 		bDestinationSet = true;
 	}
-	else if (TOptional<FGameFeaturePluginStateRange> NewDestination = StateProperties.Destination.Intersect(InDestination))
+	else if (TOptional<FGameFeaturePluginStateRange> NewDestination = This->StateProperties.Destination.Intersect(InDestination))
 	{
 		// The machine is already running so we can only transition to this range if it overlaps with our current range.
 		// We can satisfy both ranges in this case.
@@ -4171,17 +4200,17 @@ bool UGameFeaturePluginStateMachine::SetDestination(FGameFeaturePluginStateRange
 			OutCallbackHandle->Reset();
 		}
 
-		if (CurrentStateInfo.State < StateProperties.Destination)
+		if (This->CurrentStateInfo.State < This->StateProperties.Destination)
 		{
-			StateProperties.Destination = *NewDestination;
+			This->StateProperties.Destination = *NewDestination;
 
-			if (InDestination.Contains(CurrentStateInfo.State))
+			if (InDestination.Contains(This->CurrentStateInfo.State))
 			{
-				OnFeatureStateTransitionComplete.ExecuteIfBound(this, MakeValue());
+				OnFeatureStateTransitionComplete.ExecuteIfBound(This, MakeValue());
 			}
 			else
 			{
-				FDestinationGameFeaturePluginState* MinDestState = AllStates[InDestination.MinState]->AsDestinationState();
+				FDestinationGameFeaturePluginState* MinDestState = This->AllStates[InDestination.MinState]->AsDestinationState();
 				FDelegateHandle CallbackHandle = MinDestState->OnDestinationStateReached.Add(MoveTemp(OnFeatureStateTransitionComplete));
 				if (OutCallbackHandle)
 				{
@@ -4190,17 +4219,17 @@ bool UGameFeaturePluginStateMachine::SetDestination(FGameFeaturePluginStateRange
 				bDestinationChanged = true;
 			}
 		}
-		else if(CurrentStateInfo.State > StateProperties.Destination)
+		else if(This->CurrentStateInfo.State > This->StateProperties.Destination)
 		{
-			StateProperties.Destination = *NewDestination;
+			This->StateProperties.Destination = *NewDestination;
 
-			if (InDestination.Contains(CurrentStateInfo.State))
+			if (InDestination.Contains(This->CurrentStateInfo.State))
 			{
-				OnFeatureStateTransitionComplete.ExecuteIfBound(this, MakeValue());
+				OnFeatureStateTransitionComplete.ExecuteIfBound(This, MakeValue());
 			}
 			else
 			{
-				FDestinationGameFeaturePluginState* MaxDestState = AllStates[InDestination.MaxState]->AsDestinationState();
+				FDestinationGameFeaturePluginState* MaxDestState = This->AllStates[InDestination.MaxState]->AsDestinationState();
 				FDelegateHandle CallbackHandle = MaxDestState->OnDestinationStateReached.Add(MoveTemp(OnFeatureStateTransitionComplete));
 				if (OutCallbackHandle)
 				{
@@ -4222,7 +4251,7 @@ bool UGameFeaturePluginStateMachine::SetDestination(FGameFeaturePluginStateRange
 	}
 
 #if !UE_BUILD_SHIPPING
-	if (bDestinationChanged && UGameFeaturesSubsystem::Get().GetPluginDebugStateEnabled(GetPluginURL()))
+	if (bDestinationChanged && GetPluginDebugStateEnabled(GetPluginURL(This)))
 	{
 		PLATFORM_BREAK();
 	}
@@ -4231,53 +4260,53 @@ bool UGameFeaturePluginStateMachine::SetDestination(FGameFeaturePluginStateRange
 	return bDestinationSet;
 }
 
-bool UGameFeaturePluginStateMachine::TryCancel(FGameFeatureStateTransitionCanceled OnFeatureStateTransitionCanceled, FDelegateHandle* OutCallbackHandle /*= nullptr*/)
+bool UGameFeaturesSubsystem::TryCancel(UGameFeaturePluginStateMachine* This, FGameFeatureStateTransitionCanceled OnFeatureStateTransitionCanceled, FDelegateHandle* OutCallbackHandle /*= nullptr*/) const
 {
-	if (!IsRunning())
+	if (!IsRunning(This))
 	{
 		return false;
 	}
 
 	LLM_SCOPE_BYTAG(GFP);
 
-	StateProperties.bTryCancel = true;
-	FDelegateHandle CallbackHandle = StateProperties.OnTransitionCanceled.Add(MoveTemp(OnFeatureStateTransitionCanceled));
+	This->StateProperties.bTryCancel = true;
+	FDelegateHandle CallbackHandle = This->StateProperties.OnTransitionCanceled.Add(MoveTemp(OnFeatureStateTransitionCanceled));
 	if(OutCallbackHandle)
 	{
 		*OutCallbackHandle = CallbackHandle;
 	}
 
-	const EGameFeaturePluginState CurrentState = GetCurrentState();
-	AllStates[CurrentState]->TryCancelState();
+	const EGameFeaturePluginState CurrentState = GetCurrentState(This);
+	This->AllStates[CurrentState]->TryCancelState();
 
 	return true;
 }
 
-UE::GameFeatures::FResult UGameFeaturePluginStateMachine::TryUpdatePluginProtocolOptions(const FGameFeatureProtocolOptions& InOptions, bool& bOutDidUpdate)
+UE::GameFeatures::FResult UGameFeaturesSubsystem::TryUpdatePluginProtocolOptions(UGameFeaturePluginStateMachine* This, const FGameFeatureProtocolOptions& InOptions, bool& bOutDidUpdate) const
 {
 	bOutDidUpdate = false;
 
-	if (StateProperties.ProtocolOptions == InOptions)
+	if (This->StateProperties.ProtocolOptions == InOptions)
 	{
 		return MakeValue();
 	}
 
 	LLM_SCOPE_BYTAG(GFP);
 
-	const EGameFeaturePluginState CurrentState = GetCurrentState();
-	UE::GameFeatures::FResult Result = AllStates[CurrentState]->TryUpdateProtocolOptions(InOptions);
+	const EGameFeaturePluginState CurrentState = GetCurrentState(This);
+	UE::GameFeatures::FResult Result = This->AllStates[CurrentState]->TryUpdateProtocolOptions(InOptions);
 	bOutDidUpdate = Result.HasValue();
 
 	return Result;
 }
 
-void UGameFeaturePluginStateMachine::RemovePendingTransitionCallback(FDelegateHandle InHandle)
+void UGameFeaturesSubsystem::RemovePendingTransitionCallback(UGameFeaturePluginStateMachine* This, FDelegateHandle InHandle) const
 {
 	for (std::underlying_type<EGameFeaturePluginState>::type iState = 0;
 		iState < EGameFeaturePluginState::MAX;
 		++iState)
 	{
-		if (FDestinationGameFeaturePluginState* DestState = AllStates[iState]->AsDestinationState())
+		if (FDestinationGameFeaturePluginState* DestState = This->AllStates[iState]->AsDestinationState())
 		{
 			if (DestState->OnDestinationStateReached.Remove(InHandle))
 			{
@@ -4287,13 +4316,13 @@ void UGameFeaturePluginStateMachine::RemovePendingTransitionCallback(FDelegateHa
 	}
 }
 
-void UGameFeaturePluginStateMachine::RemovePendingTransitionCallback(FDelegateUserObject DelegateObject)
+void UGameFeaturesSubsystem::RemovePendingTransitionCallback(UGameFeaturePluginStateMachine* This, FDelegateUserObject DelegateObject) const
 {
 	for (std::underlying_type<EGameFeaturePluginState>::type iState = 0;
 		iState < EGameFeaturePluginState::MAX;
 		++iState)
 	{
-		if (FDestinationGameFeaturePluginState* DestState = AllStates[iState]->AsDestinationState())
+		if (FDestinationGameFeaturePluginState* DestState = This->AllStates[iState]->AsDestinationState())
 		{
 			if (DestState->OnDestinationStateReached.RemoveAll(DelegateObject))
 			{
@@ -4303,215 +4332,237 @@ void UGameFeaturePluginStateMachine::RemovePendingTransitionCallback(FDelegateUs
 	}
 }
 
-void UGameFeaturePluginStateMachine::RemovePendingCancelCallback(FDelegateHandle InHandle)
+void UGameFeaturesSubsystem::RemovePendingCancelCallback(UGameFeaturePluginStateMachine* This, FDelegateHandle InHandle) const
 {
-	StateProperties.OnTransitionCanceled.Remove(InHandle);
+	This->StateProperties.OnTransitionCanceled.Remove(InHandle);
 }
 
-void UGameFeaturePluginStateMachine::RemovePendingCancelCallback(FDelegateUserObject DelegateObject)
+void UGameFeaturesSubsystem::RemovePendingCancelCallback(UGameFeaturePluginStateMachine* This, FDelegateUserObject DelegateObject) const
 {
-	StateProperties.OnTransitionCanceled.RemoveAll(DelegateObject);
+	This->StateProperties.OnTransitionCanceled.RemoveAll(DelegateObject);
 }
 
-const FString& UGameFeaturePluginStateMachine::GetGameFeatureName() const
+const FString& UGameFeaturesSubsystem::GetGameFeatureName(const UGameFeaturePluginStateMachine* This) const
 {
 	FString PluginFilename;
-	if (!StateProperties.PluginName.IsEmpty())
+	if (!This->StateProperties.PluginName.IsEmpty())
 	{
-		return StateProperties.PluginName;
+		return This->StateProperties.PluginName;
 	}
 	else
 	{
-		return StateProperties.PluginIdentifier.GetFullPluginURL();
+		return This->StateProperties.PluginIdentifier.GetFullPluginURL();
 	}
 }
 
-const FGameFeaturePluginIdentifier& UGameFeaturePluginStateMachine::GetPluginIdentifier() const
+const FGameFeaturePluginIdentifier& UGameFeaturesSubsystem::GetPluginIdentifier(const UGameFeaturePluginStateMachine* This) const
 {
-	return StateProperties.PluginIdentifier;
+	return This->StateProperties.PluginIdentifier;
 }
 
-const FString& UGameFeaturePluginStateMachine::GetPluginURL() const
+const FString& UGameFeaturesSubsystem::GetPluginURL(const UGameFeaturePluginStateMachine* This) const
 {
-	return StateProperties.PluginIdentifier.GetFullPluginURL();
+	return This->StateProperties.PluginIdentifier.GetFullPluginURL();
 }
 
-const FGameFeatureProtocolMetadata& UGameFeaturePluginStateMachine::GetProtocolMetadata() const
+const FGameFeatureProtocolMetadata& UGameFeaturesSubsystem::GetProtocolMetadata(const UGameFeaturePluginStateMachine* This) const
 {
-	return StateProperties.ProtocolMetadata;
+	return This->StateProperties.ProtocolMetadata;
 }
 
-const FGameFeatureProtocolOptions& UGameFeaturePluginStateMachine::GetProtocolOptions() const
+const FGameFeatureProtocolOptions& UGameFeaturesSubsystem::GetProtocolOptions(const UGameFeaturePluginStateMachine* This) const
 {
-	return StateProperties.ProtocolOptions;
+	return This->StateProperties.ProtocolOptions;
 }
 
-FGameFeatureProtocolOptions UGameFeaturePluginStateMachine::RecycleProtocolOptions() const
+FGameFeatureProtocolOptions UGameFeaturesSubsystem::RecycleProtocolOptions(const UGameFeaturePluginStateMachine* This) const
 {
-	return StateProperties.RecycleProtocolOptions();
+	return This->StateProperties.RecycleProtocolOptions();
 }
 
-const FString& UGameFeaturePluginStateMachine::GetPluginName() const
+const FString& UGameFeaturesSubsystem::GetPluginName(const UGameFeaturePluginStateMachine* This) const
 {
-	return StateProperties.PluginName;
+	return This->StateProperties.PluginName;
 }
 
-bool UGameFeaturePluginStateMachine::GetPluginFilename(FString& OutPluginFilename) const
+bool UGameFeaturesSubsystem::GetPluginFilename(const UGameFeaturePluginStateMachine* This, FString& OutPluginFilename) const
 {
-	OutPluginFilename = StateProperties.PluginInstalledFilename;
+	OutPluginFilename = This->StateProperties.PluginInstalledFilename;
 	return !OutPluginFilename.IsEmpty();
 }
 
-EGameFeaturePluginState UGameFeaturePluginStateMachine::GetCurrentState() const
+EGameFeaturePluginState UGameFeaturesSubsystem::GetCurrentState(const UGameFeaturePluginStateMachine* This) const
 {
-	return GetCurrentStateInfo().State;
+	return GetCurrentStateInfo(This).State;
 }
 
-FGameFeaturePluginStateRange UGameFeaturePluginStateMachine::GetDestination() const
+FGameFeaturePluginStateRange UGameFeaturesSubsystem::GetDestination(const UGameFeaturePluginStateMachine* This) const
 {
-	return StateProperties.Destination;
+	return This->StateProperties.Destination;
 }
 
-const FGameFeaturePluginStateInfo& UGameFeaturePluginStateMachine::GetCurrentStateInfo() const
+const FGameFeaturePluginStateInfo& UGameFeaturesSubsystem::GetCurrentStateInfo(const UGameFeaturePluginStateMachine* This) const
 {
-	return CurrentStateInfo;
+	return This->CurrentStateInfo;
 }
 
-bool UGameFeaturePluginStateMachine::IsRunning() const
+bool UGameFeaturesSubsystem::IsRunning(const UGameFeaturePluginStateMachine* This) const
 {
-	return !StateProperties.Destination.Contains(CurrentStateInfo.State);
+	return !This->StateProperties.Destination.Contains(This->CurrentStateInfo.State);
 }
 
-bool UGameFeaturePluginStateMachine::IsStatusKnown() const
+bool UGameFeaturesSubsystem::IsStatusKnown(const UGameFeaturePluginStateMachine* This) const
 {
-	return GetCurrentState() == EGameFeaturePluginState::ErrorUnavailable || 
-		   GetCurrentState() == EGameFeaturePluginState::Uninstalling || 
-		   GetCurrentState() == EGameFeaturePluginState::ErrorUninstalling ||
-		   GetCurrentState() >= EGameFeaturePluginState::StatusKnown;
+	return GetCurrentState(This) == EGameFeaturePluginState::ErrorUnavailable || 
+		   GetCurrentState(This) == EGameFeaturePluginState::Uninstalling || 
+		   GetCurrentState(This) == EGameFeaturePluginState::ErrorUninstalling ||
+		   GetCurrentState(This) >= EGameFeaturePluginState::StatusKnown;
 }
 
-bool UGameFeaturePluginStateMachine::IsAvailable() const
+bool UGameFeaturesSubsystem::IsAvailable(const UGameFeaturePluginStateMachine* This) const
 {
-	ensure(IsStatusKnown());
-	return GetCurrentState() >= EGameFeaturePluginState::StatusKnown;
+	ensure(IsStatusKnown(This));
+	return GetCurrentState(This) >= EGameFeaturePluginState::StatusKnown;
 }
 
-bool UGameFeaturePluginStateMachine::IsInErrorState() const
+bool UGameFeaturesSubsystem::IsInErrorState(const UGameFeaturePluginStateMachine* This) const
 {
-	return IsValidErrorState(GetCurrentState());
+	return IsValidErrorState(This, GetCurrentState(This));
 }
 
-bool UGameFeaturePluginStateMachine::AllowAsyncLoading() const
+bool UGameFeaturesSubsystem::AllowAsyncLoading(const UGameFeaturePluginStateMachine* This) const
 {
-	return StateProperties.AllowAsyncLoading();
+	return This->StateProperties.AllowAsyncLoading();
 }
 
-bool UGameFeaturePluginStateMachine::HasAssetStreamingDependencies() const
+bool UGameFeaturesSubsystem::HasAssetStreamingDependencies(const UGameFeaturePluginStateMachine* This) const
 {
-	ensure(IsStatusKnown());
-	if (StateProperties.ProtocolMetadata.HasSubtype<FInstallBundlePluginProtocolMetaData>())
+	ensure(IsStatusKnown(This));
+	if (This->StateProperties.ProtocolMetadata.HasSubtype<FInstallBundlePluginProtocolMetaData>())
 	{
-		const FInstallBundlePluginProtocolMetaData& ProtocolData = StateProperties.ProtocolMetadata.GetSubtype<FInstallBundlePluginProtocolMetaData>();
+		const FInstallBundlePluginProtocolMetaData& ProtocolData = This->StateProperties.ProtocolMetadata.GetSubtype<FInstallBundlePluginProtocolMetaData>();
 		return !ProtocolData.InstallBundlesWithAssetDependencies.IsEmpty();
 	}
 	return false;
 }
 
-void UGameFeaturePluginStateMachine::SetWasLoadedAsBuiltIn()
+void UGameFeaturesSubsystem::SetWasLoadedAsBuiltIn(UGameFeaturePluginStateMachine* This)
 {
-	StateProperties.bWasLoadedAsBuiltInGameFeaturePlugin = true;
+	This->StateProperties.bWasLoadedAsBuiltInGameFeaturePlugin = true;
 }
 
-bool UGameFeaturePluginStateMachine::WasLoadedAsBuiltIn() const
+bool UGameFeaturesSubsystem::WasLoadedAsBuiltIn(const UGameFeaturePluginStateMachine* This) const
 {
-	return StateProperties.bWasLoadedAsBuiltInGameFeaturePlugin;
+	return This->StateProperties.bWasLoadedAsBuiltInGameFeaturePlugin;
 }
 
-UGameFeatureData* UGameFeaturePluginStateMachine::GetGameFeatureDataForActivePlugin()
+UGameFeatureData* UGameFeaturesSubsystem::GetGameFeatureDataForActivePlugin(UGameFeaturePluginStateMachine* This) const
 {
-	if (GetCurrentState() == EGameFeaturePluginState::Active)
+	if (GetCurrentState(This) == EGameFeaturePluginState::Active)
+	{
+		return This->StateProperties.GameFeatureData;
+	}
+
+	return nullptr;
+}
+
+const UGameFeatureData* UGameFeaturesSubsystem::GetGameFeatureDataForActivePlugin(const UGameFeaturePluginStateMachine* This) const
+{
+	if (GetCurrentState(This) == EGameFeaturePluginState::Active)
+	{
+		return This->StateProperties.GameFeatureData;
+	}
+
+	return nullptr;
+}
+
+UGameFeatureData* UGameFeaturesSubsystem::GetGameFeatureDataForRegisteredPlugin(UGameFeaturePluginStateMachine* This, bool bCheckForRegistering /*= false*/) const
+{
+	const EGameFeaturePluginState CurrentState = GetCurrentState(This);
+
+	if (CurrentState >= EGameFeaturePluginState::Registered || (bCheckForRegistering && (CurrentState == EGameFeaturePluginState::Registering)))
 	{
-		return StateProperties.GameFeatureData;
+		return This->StateProperties.GameFeatureData;
 	}
 
 	return nullptr;
 }
 
-UGameFeatureData* UGameFeaturePluginStateMachine::GetGameFeatureDataForRegisteredPlugin(bool bCheckForRegistering /*= false*/)
+const UGameFeatureData* UGameFeaturesSubsystem::GetGameFeatureDataForRegisteredPlugin(const UGameFeaturePluginStateMachine* This, bool bCheckForRegistering /*= false*/) const
 {
-	const EGameFeaturePluginState CurrentState = GetCurrentState();
+	const EGameFeaturePluginState CurrentState = GetCurrentState(This);
 
 	if (CurrentState >= EGameFeaturePluginState::Registered || (bCheckForRegistering && (CurrentState == EGameFeaturePluginState::Registering)))
 	{
-		return StateProperties.GameFeatureData;
+		return This->StateProperties.GameFeatureData;
 	}
 
 	return nullptr;
 }
 
-const FGameFeaturePluginStateMachineProperties& UGameFeaturePluginStateMachine::GetProperties() const
+const FGameFeaturePluginStateMachineProperties& UGameFeaturesSubsystem::GetProperties(const UGameFeaturePluginStateMachine* This) const
 {
-	return StateProperties;
+	return This->StateProperties;
 }
 
-bool UGameFeaturePluginStateMachine::IsErrorStateUnrecoverable() const
+bool UGameFeaturesSubsystem::IsErrorStateUnrecoverable(const UGameFeaturePluginStateMachine* This) const
 {
-	return bIsInUnrecoverableError;
+	return This->bIsInUnrecoverableError;
 }
 
-void UGameFeaturePluginStateMachine::SetUnrecoverableError()
+void UGameFeaturesSubsystem::SetUnrecoverableError(UGameFeaturePluginStateMachine* This) const
 {
-	bIsInUnrecoverableError = true;
+	This->bIsInUnrecoverableError = true;
 }
 
-bool UGameFeaturePluginStateMachine::IsValidTransitionState(EGameFeaturePluginState InState) const
+bool UGameFeaturesSubsystem::IsValidTransitionState(const UGameFeaturePluginStateMachine* This, EGameFeaturePluginState InState) const
 {
 	check(InState != EGameFeaturePluginState::MAX);
-	return AllStates[InState]->GetStateType() == EGameFeaturePluginStateType::Transition;
+	return This->AllStates[InState]->GetStateType() == EGameFeaturePluginStateType::Transition;
 }
 
-bool UGameFeaturePluginStateMachine::IsValidDestinationState(EGameFeaturePluginState InDestinationState) const
+bool UGameFeaturesSubsystem::IsValidDestinationState(const UGameFeaturePluginStateMachine* This, EGameFeaturePluginState InDestinationState) const
 {
 	check(InDestinationState != EGameFeaturePluginState::MAX);
-	return AllStates[InDestinationState]->GetStateType() == EGameFeaturePluginStateType::Destination;
+	return This->AllStates[InDestinationState]->GetStateType() == EGameFeaturePluginStateType::Destination;
 }
 
-bool UGameFeaturePluginStateMachine::IsValidErrorState(EGameFeaturePluginState InDestinationState) const
+bool UGameFeaturesSubsystem::IsValidErrorState(const UGameFeaturePluginStateMachine* This, EGameFeaturePluginState InDestinationState) const
 {
 	check(InDestinationState != EGameFeaturePluginState::MAX);
-	return AllStates[InDestinationState]->GetStateType() == EGameFeaturePluginStateType::Error;
+	return This->AllStates[InDestinationState]->GetStateType() == EGameFeaturePluginStateType::Error;
 }
 
 UE_TRACE_EVENT_BEGIN(Cpu, GFP_UpdateStateMachine, NoSync)
 	UE_TRACE_EVENT_FIELD(UE::Trace::WideString, PluginName)
 UE_TRACE_EVENT_END()
 
-void UGameFeaturePluginStateMachine::UpdateStateMachine()
+void UGameFeaturesSubsystem::UpdateStateMachine(UGameFeaturePluginStateMachine* This) 
 {
 	LLM_SCOPE_BYTAG(GFP);
 
-	const EGameFeaturePluginState InitialState = GetCurrentState();
+	const EGameFeaturePluginState InitialState = GetCurrentState(This);
 	EGameFeaturePluginState CurrentState = InitialState;
-	if (bInUpdateStateMachine)
+	if (This->bInUpdateStateMachine)
 	{
-		UE_LOG(LogGameFeatures, Verbose, TEXT("Game feature state machine skipping update for %s in ::UpdateStateMachine. Current State: %s"), *GetGameFeatureName(), *UE::GameFeatures::ToString(CurrentState));
+		UE_LOG(LogGameFeatures, Verbose, TEXT("Game feature state machine skipping update for %s in ::UpdateStateMachine. Current State: %s"), *GetGameFeatureName(This), *UE::GameFeatures::ToString(CurrentState));
 		return;
 	}
 
 	UE_TRACE_LOG_SCOPED_T(Cpu, GFP_UpdateStateMachine, CpuChannel)
-		<< GFP_UpdateStateMachine.PluginName(*GetGameFeatureName());
+		<< GFP_UpdateStateMachine.PluginName(*GetGameFeatureName(This));
 
-	TOptional<TGuardValue<bool>> ScopeGuard(InPlace, bInUpdateStateMachine, true);
+	TOptional<TGuardValue<bool>> ScopeGuard(InPlace, This->bInUpdateStateMachine, true);
 	
 	using StateIt = std::underlying_type<EGameFeaturePluginState>::type;
 
-	auto DoCallbacks = [this](const UE::GameFeatures::FResult& Result, StateIt Begin, StateIt End)
+	auto DoCallbacks = [This](const UE::GameFeatures::FResult& Result, StateIt Begin, StateIt End)
 	{
 		TRACE_CPUPROFILER_EVENT_SCOPE(GFP_UpdateStateMachine_DoCallbacks);
 
 		for (StateIt iState = Begin; iState < End; ++iState)
 		{
-			if (FDestinationGameFeaturePluginState* DestState = AllStates[iState]->AsDestinationState())
+			if (FDestinationGameFeaturePluginState* DestState = This->AllStates[iState]->AsDestinationState())
 			{
 				// Use a local callback on the stack. If SetDestination() is called from the callback then we don't want to stomp the callback
 				// for the new state transition request.
@@ -4519,7 +4570,7 @@ void UGameFeaturePluginStateMachine::UpdateStateMachine()
 				UE::GameFeatures::FBroadcastingOnDestinationStateReached LocalOnDestinationStateReached(MoveTemp(DestState->OnDestinationStateReached));
 				DestState->OnDestinationStateReached.Clear();
 
-				LocalOnDestinationStateReached.CallbackDelegate.Broadcast(this, Result);
+				LocalOnDestinationStateReached.CallbackDelegate.Broadcast(This, Result);
 			}
 		}
 	};
@@ -4529,7 +4580,7 @@ void UGameFeaturePluginStateMachine::UpdateStateMachine()
 		DoCallbacks(Result, InState, InState + 1);
 	};
 
-	RegisterAsTransitioningStateMachine();
+	RegisterAsTransitioningStateMachine(This);
 
 	bool bKeepProcessing = false;
 	int32 NumTransitions = 0;
@@ -4540,39 +4591,39 @@ void UGameFeaturePluginStateMachine::UpdateStateMachine()
 
 		FGameFeaturePluginStateStatus StateStatus;
 		{
-			const FGameFeaturePluginStateMachineProperties& Props = AllStates[CurrentState]->StateProperties;
+			const FGameFeaturePluginStateMachineProperties& Props = This->AllStates[CurrentState]->StateProperties;
 			FName LoaderName = Props.GameFeatureData ? Props.GameFeatureData->GetPackage()->GetFName() : FName(Props.PluginName);
 			UE_TRACK_REFERENCING_PACKAGE_SCOPED(LoaderName, UE::GameFeatures::GetStateName(CurrentState));
 
 			TRACE_CPUPROFILER_EVENT_SCOPE(GFP_UpdateStateMachine_UpdateState);
-			AllStates[CurrentState]->UpdateState(StateStatus);
+			This->AllStates[CurrentState]->UpdateState(StateStatus);
 		}
 
 		if (StateStatus.TransitionToState == CurrentState)
 		{
-			UE_LOG(LogGameFeatures, Fatal, TEXT("Game feature state %s transitioning to itself. GameFeature: %s"), *UE::GameFeatures::ToString(CurrentState), *GetGameFeatureName());
+			UE_LOG(LogGameFeatures, Fatal, TEXT("Game feature state %s transitioning to itself. GameFeature: %s"), *UE::GameFeatures::ToString(CurrentState), *GetGameFeatureName(This));
 		}
 
 		if (StateStatus.TransitionToState != EGameFeaturePluginState::Uninitialized)
 		{
-			UE_LOG(LogGameFeatures, Verbose, TEXT("Game feature '%s' transitioning state (%s -> %s)"), *GetGameFeatureName(), *UE::GameFeatures::ToString(CurrentState), *UE::GameFeatures::ToString(StateStatus.TransitionToState));
+			UE_LOG(LogGameFeatures, Verbose, TEXT("Game feature '%s' transitioning state (%s -> %s)"), *GetGameFeatureName(This), *UE::GameFeatures::ToString(CurrentState), *UE::GameFeatures::ToString(StateStatus.TransitionToState));
 			{
 				TRACE_CPUPROFILER_EVENT_SCOPE(GFP_UpdateStateMachine_EndState);
-				AllStates[CurrentState]->EndState();
-				CheckAndCancelBatchingRequestForCurrentState();
+				This->AllStates[CurrentState]->EndState();
+				CheckAndCancelBatchingRequestForCurrentState(This);
 			}
-			CurrentStateInfo = FGameFeaturePluginStateInfo(StateStatus.TransitionToState);
+			This->CurrentStateInfo = FGameFeaturePluginStateInfo(StateStatus.TransitionToState);
 			CurrentState = StateStatus.TransitionToState;
 			check(CurrentState != EGameFeaturePluginState::MAX);
 
-			const FGameFeaturePluginStateMachineProperties& Props = AllStates[CurrentState]->StateProperties;
+			const FGameFeaturePluginStateMachineProperties& Props = This->AllStates[CurrentState]->StateProperties;
 			FName LoaderName = Props.GameFeatureData ? Props.GameFeatureData->GetPackage()->GetFName() : FName(Props.PluginName);
 			UE_TRACK_REFERENCING_PACKAGE_SCOPED(LoaderName, UE::GameFeatures::GetStateName(CurrentState));
 
 			{
 				TRACE_CPUPROFILER_EVENT_SCOPE(GFP_UpdateStateMachine_BeginState);
-				CheckAddBatchingRequestForCurrentState();
-				AllStates[CurrentState]->BeginState();
+				CheckAddBatchingRequestForCurrentState(This);
+				This->AllStates[CurrentState]->BeginState();
 			}
 
 			if (CurrentState == EGameFeaturePluginState::Terminal)
@@ -4581,35 +4632,35 @@ void UGameFeaturePluginStateMachine::UpdateStateMachine()
 
 				// Remove from gamefeature subsystem before calling back in case this GFP is reloaded on callback,
 				// but make sure we don't get destroyed from a GC during a callback
-				UGameFeaturesSubsystem::Get().BeginTermination(this);
+				BeginTermination(This);
 			}
 
-			if (StateProperties.bTryCancel && AllStates[CurrentState]->GetStateType() != EGameFeaturePluginStateType::Transition)
+			if (This->StateProperties.bTryCancel && This->AllStates[CurrentState]->GetStateType() != EGameFeaturePluginStateType::Transition)
 			{
-				StateProperties.Destination = FGameFeaturePluginStateRange(CurrentState);
+				This->StateProperties.Destination = FGameFeaturePluginStateRange(CurrentState);
 
-				StateProperties.bTryCancel = false;
+				This->StateProperties.bTryCancel = false;
 				bKeepProcessing = false;
 
-				// Make sure bInUpdateStateMachine is not set while processing callbacks if we are at our destination
+				// Make sure This->bInUpdateStateMachine is not set while processing callbacks if we are at our destination
 				ScopeGuard.Reset();
 
 				// For all callbacks, return the CanceledResult
 				DoCallbacks(UE::GameFeatures::CanceledResult, 0, EGameFeaturePluginState::MAX);
 
 				// Must be called after transtition callbacks, UGameFeaturesSubsystem::ChangeGameFeatureTargetStateComplete may remove the this machine from the subsystem
-				UE::GameFeatures::FBroadcastingOnTransitionCanceled LocalOnTransitionCanceled(MoveTemp(StateProperties.OnTransitionCanceled));
-				StateProperties.OnTransitionCanceled.Clear();
-				LocalOnTransitionCanceled.CallbackDelegate.Broadcast(this);
+				UE::GameFeatures::FBroadcastingOnTransitionCanceled LocalOnTransitionCanceled(MoveTemp(This->StateProperties.OnTransitionCanceled));
+				This->StateProperties.OnTransitionCanceled.Clear();
+				LocalOnTransitionCanceled.CallbackDelegate.Broadcast(This);
 			}
 			else if (const bool bError = !StateStatus.TransitionResult.HasValue(); bError)
 			{
-				check(IsValidErrorState(CurrentState));
-				StateProperties.Destination = FGameFeaturePluginStateRange(CurrentState);
+				check(IsValidErrorState(This, CurrentState));
+				This->StateProperties.Destination = FGameFeaturePluginStateRange(CurrentState);
 
 				bKeepProcessing = false;
 				
-				// Make sure bInUpdateStateMachine is not set while processing callbacks if we are at our destination
+				// Make sure This->bInUpdateStateMachine is not set while processing callbacks if we are at our destination
 				ScopeGuard.Reset();
 
 				// In case of an error, callback all possible callbacks
@@ -4617,10 +4668,10 @@ void UGameFeaturePluginStateMachine::UpdateStateMachine()
 			}
 			else
 			{
-				bKeepProcessing = AllStates[CurrentState]->GetStateType() == EGameFeaturePluginStateType::Transition || !StateProperties.Destination.Contains(CurrentState);
+				bKeepProcessing = This->AllStates[CurrentState]->GetStateType() == EGameFeaturePluginStateType::Transition || !This->StateProperties.Destination.Contains(CurrentState);
 				if (!bKeepProcessing)
 				{
-					// Make sure bInUpdateStateMachine is not set while processing callbacks if we are at our destination
+					// Make sure This->bInUpdateStateMachine is not set while processing callbacks if we are at our destination
 					ScopeGuard.Reset();
 				}
 
@@ -4629,24 +4680,24 @@ void UGameFeaturePluginStateMachine::UpdateStateMachine()
 
 			if (!bKeepProcessing)
 			{
-				UnregisterAsTransitioningStateMachine();
+				UnregisterAsTransitioningStateMachine(This);
 			}
 
 			if (CurrentState == EGameFeaturePluginState::Terminal)
 			{
 				TRACE_CPUPROFILER_EVENT_SCOPE(GFP_UpdateStateMachine_FinishTerm);
 
-				UnregisterAsTransitioningStateMachine();
+				UnregisterAsTransitioningStateMachine(This);
 
 				check(bKeepProcessing == false);
 				// Now that callbacks are done this machine can be cleaned up
-				UGameFeaturesSubsystem::Get().FinishTermination(this);
-				MarkAsGarbage();
+				FinishTermination(This);
+				This->MarkAsGarbage();
 			}
 		}
-		else if(!IsRunning())
+		else if(!IsRunning(This))
 		{
-			UnregisterAsTransitioningStateMachine();
+			UnregisterAsTransitioningStateMachine(This);
 		}
 
 		// Log our final state if we've finished transitioning
@@ -4658,103 +4709,103 @@ void UGameFeaturePluginStateMachine::UpdateStateMachine()
 				if (StateStatus.bSuppressErrorLog)
 				{
 					UE_LOG(LogGameFeatures, Display, ErrLogFmt,
-						*GetGameFeatureName(),
+						*GetGameFeatureName(This),
 						*UE::GameFeatures::ToString(CurrentState),
-						*UE::GameFeatures::ToString(StateProperties.Destination.MinState),
-						*UE::GameFeatures::ToString(StateProperties.Destination.MaxState),
+						*UE::GameFeatures::ToString(This->StateProperties.Destination.MinState),
+						*UE::GameFeatures::ToString(This->StateProperties.Destination.MaxState),
 						*UE::GameFeatures::ToString(StateStatus.TransitionResult));
 				}
 				else
 				{
 					UE_LOG(LogGameFeatures, Error, ErrLogFmt,
-						*GetGameFeatureName(),
+						*GetGameFeatureName(This),
 						*UE::GameFeatures::ToString(CurrentState),
-						*UE::GameFeatures::ToString(StateProperties.Destination.MinState),
-						*UE::GameFeatures::ToString(StateProperties.Destination.MaxState),
+						*UE::GameFeatures::ToString(This->StateProperties.Destination.MinState),
+						*UE::GameFeatures::ToString(This->StateProperties.Destination.MaxState),
 						*UE::GameFeatures::ToString(StateStatus.TransitionResult));
 				}
 			}
-			else if (StateProperties.Destination.Contains(CurrentState))
+			else if (This->StateProperties.Destination.Contains(CurrentState))
 			{
 				UE_LOG(LogGameFeatures, Display, TEXT("Game feature '%s' transitioned successfully. Ending state: %s [%s, %s]"),
-					*GetGameFeatureName(),
+					*GetGameFeatureName(This),
 					*UE::GameFeatures::ToString(CurrentState),
-					*UE::GameFeatures::ToString(StateProperties.Destination.MinState),
-					*UE::GameFeatures::ToString(StateProperties.Destination.MaxState));
+					*UE::GameFeatures::ToString(This->StateProperties.Destination.MinState),
+					*UE::GameFeatures::ToString(This->StateProperties.Destination.MaxState));
 			}
 		}
 
 		if (NumTransitions++ > MaxTransitions)
 		{
-			UE_LOG(LogGameFeatures, Fatal, TEXT("Infinite loop in game feature state machine transitions. Current state %s. GameFeature: %s"), *UE::GameFeatures::ToString(CurrentState), *GetGameFeatureName());
+			UE_LOG(LogGameFeatures, Fatal, TEXT("Infinite loop in game feature state machine transitions. Current state %s. GameFeature: %s"), *UE::GameFeatures::ToString(CurrentState), *GetGameFeatureName(This));
 		}
 	} while (bKeepProcessing);
 }
 
-void UGameFeaturePluginStateMachine::UpdateCurrentStateProgress(float Progress)
+void UGameFeaturesSubsystem::UpdateCurrentStateProgress(UGameFeaturePluginStateMachine* This, float Progress) const
 {
-	CurrentStateInfo.Progress = Progress;
+	This->CurrentStateInfo.Progress = Progress;
 }
 
-void UGameFeaturePluginStateMachine::RegisterAsTransitioningStateMachine()
+void UGameFeaturesSubsystem::RegisterAsTransitioningStateMachine(UGameFeaturePluginStateMachine* This) 
 {
-	if (bRegisteredAsTransitioningGFPSM)
+	if (This->bRegisteredAsTransitioningGFPSM)
 	{
 		return;
 	}
 
-	UGameFeaturesSubsystem::Get().RegisterRunningStateMachine(this);
-	bRegisteredAsTransitioningGFPSM = true;
+	RegisterRunningStateMachine(This);
+	This->bRegisteredAsTransitioningGFPSM = true;
 }
 
-void UGameFeaturePluginStateMachine::UnregisterAsTransitioningStateMachine()
+void UGameFeaturesSubsystem::UnregisterAsTransitioningStateMachine(UGameFeaturePluginStateMachine* This) 
 {
-	if (!bRegisteredAsTransitioningGFPSM)
+	if (!This->bRegisteredAsTransitioningGFPSM)
 	{
 		return;
 	}
 
-	UGameFeaturesSubsystem::Get().UnregisterRunningStateMachine(this);
-	bRegisteredAsTransitioningGFPSM = false;
+	UnregisterRunningStateMachine(This);
+	This->bRegisteredAsTransitioningGFPSM = false;
 }
 
-void UGameFeaturePluginStateMachine::CheckAddBatchingRequestForCurrentState()
+void UGameFeaturesSubsystem::CheckAddBatchingRequestForCurrentState(UGameFeaturePluginStateMachine* This) 
 {
-	check(!StateProperties.BatchProcessingHandle.IsValid());
+	check(!This->StateProperties.BatchProcessingHandle.IsValid());
 
 	const bool bCanBatchProcess =
-		UE::GameFeatures::BatchProcessingHelperFunctors[CurrentStateInfo.State].ImplementsBatchProcess() &&
-		StateProperties.CanBatchProcess() &&
-		AllStates[CurrentStateInfo.State]->CanBatchProcess();
+		UE::GameFeatures::BatchProcessingHelperFunctors[This->CurrentStateInfo.State].ImplementsBatchProcess() &&
+		This->StateProperties.CanBatchProcess() &&
+		This->AllStates[This->CurrentStateInfo.State]->CanBatchProcess();
 	
 	if (bCanBatchProcess)
 	{
-		UE_LOG(LogGameFeatures, Verbose, TEXT("Game feature '%s' awaiting batch processing of state (%s)"), *GetGameFeatureName(), *UE::GameFeatures::ToString(CurrentStateInfo.State));
-		StateProperties.BatchProcessingHandle = UGameFeaturesSubsystem::Get().AddBatchingRequest(
-			CurrentStateInfo.State,
-			StateProperties.OnRequestUpdateStateMachine);
+		UE_LOG(LogGameFeatures, Verbose, TEXT("Game feature '%s' awaiting batch processing of state (%s)"), *GetGameFeatureName(This), *UE::GameFeatures::ToString(This->CurrentStateInfo.State));
+		This->StateProperties.BatchProcessingHandle = AddBatchingRequest(
+			This->CurrentStateInfo.State,
+			This->StateProperties.OnRequestUpdateStateMachine);
 	}
 }
 
-void UGameFeaturePluginStateMachine::CheckAndCancelBatchingRequestForCurrentState()
+void UGameFeaturesSubsystem::CheckAndCancelBatchingRequestForCurrentState(UGameFeaturePluginStateMachine* This) 
 {
 	// Reset batch processing state.
-	StateProperties.bWasBatchProcessed = false;
+	This->StateProperties.bWasBatchProcessed = false;
 
 	// If we are currently awaiting batch processing, cancel and update the state machine.
-	if (StateProperties.IsWaitingForBatchProcessing())
+	if (This->StateProperties.IsWaitingForBatchProcessing())
 	{
-		UE_LOG(LogGameFeatures, Verbose, TEXT("Game feature '%s' cancelled batch processing of state (%s)"), *GetGameFeatureName(), *UE::GameFeatures::ToString(CurrentStateInfo.State));
-		UGameFeaturesSubsystem::Get().CancelBatchingRequest(
-			CurrentStateInfo.State,
-			StateProperties.BatchProcessingHandle);
+		UE_LOG(LogGameFeatures, Verbose, TEXT("Game feature '%s' cancelled batch processing of state (%s)"), *GetGameFeatureName(This), *UE::GameFeatures::ToString(This->CurrentStateInfo.State));
+		CancelBatchingRequest(
+			This->CurrentStateInfo.State,
+			This->StateProperties.BatchProcessingHandle);
 
-		StateProperties.BatchProcessingHandle.Reset();
-		UpdateStateMachine();
+		This->StateProperties.BatchProcessingHandle.Reset();
+		UpdateStateMachine(This);
 	}
 }
 
-/* static */ void UGameFeaturePluginStateMachine::BatchProcess(EGameFeaturePluginState State, TConstArrayView<UGameFeaturePluginStateMachine*> GFPSMs)
+/* static */ void UGameFeaturesSubsystem::BatchProcess(EGameFeaturePluginState State, TConstArrayView<UGameFeaturePluginStateMachine*> GFPSMs)
 {
 	UE::GameFeatures::BatchProcessingHelperFunctors[State].BatchProcess(GFPSMs);
 	for (UGameFeaturePluginStateMachine* GFPSM : GFPSMs)
@@ -4764,20 +4815,20 @@ void UGameFeaturePluginStateMachine::CheckAndCancelBatchingRequestForCurrentStat
 	}
 }
 
-void UGameFeaturePluginStateMachine::ExcludeFromBatchProcessing()
+void UGameFeaturesSubsystem::ExcludeFromBatchProcessing(UGameFeaturePluginStateMachine* This) 
 {
 	// Ensure protocol options are updated to reflect the exclusion of this state machine from batch processing.
-	if (StateProperties.ProtocolOptions.bBatchProcess)
+	if (This->StateProperties.ProtocolOptions.bBatchProcess)
 	{
-		UE_LOG(LogGameFeatures, Verbose, TEXT("Game feature '%s' excluded from batch processing"), *GetGameFeatureName());
+		UE_LOG(LogGameFeatures, Verbose, TEXT("Game feature '%s' excluded from batch processing"), *GetGameFeatureName(This));
 
-		FGameFeatureProtocolOptions NewOptions = StateProperties.ProtocolOptions;
+		FGameFeatureProtocolOptions NewOptions = This->StateProperties.ProtocolOptions;
 		NewOptions.bBatchProcess = false;
 
 		bool bOutDidUpdate = false;
-		TryUpdatePluginProtocolOptions(NewOptions, bOutDidUpdate);
+		TryUpdatePluginProtocolOptions(This, NewOptions, bOutDidUpdate);
 		check(bOutDidUpdate);
-		CheckAndCancelBatchingRequestForCurrentState();
+		CheckAndCancelBatchingRequestForCurrentState(This);
 	}	
 }
 
diff --git a/Engine/Plugins/Runtime/GameFeatures/Source/GameFeatures/Private/GameFeaturePluginStateMachine.h b/Engine/Plugins/Runtime/GameFeatures/Source/GameFeatures/Private/GameFeaturePluginStateMachine.h
--- a/Engine/Plugins/Runtime/GameFeatures/Source/GameFeatures/Private/GameFeaturePluginStateMachine.h
+++ b/Engine/Plugins/Runtime/GameFeatures/Source/GameFeatures/Private/GameFeaturePluginStateMachine.h
@@ -369,6 +369,7 @@ private:
 	bool bSuppressErrorLog = false;
 
 	friend class UGameFeaturePluginStateMachine;
+	friend class UGameFeaturesSubsystem;
 
 public:
 	void SetTransition(EGameFeaturePluginState InTransitionToState);
@@ -388,7 +389,11 @@ struct FErrorGameFeaturePluginState;
 /** Base class for all game feature plugin states */
 struct FGameFeaturePluginState
 {
-	FGameFeaturePluginState(FGameFeaturePluginStateMachineProperties& InStateProperties) : StateProperties(InStateProperties) {}
+	FGameFeaturePluginState(EGameFeaturePluginState InState, FGameFeaturePluginStateMachineProperties& InStateProperties) 
+		: StateProperties(InStateProperties)
+		, State(InState) 
+	{
+	}
 	virtual ~FGameFeaturePluginState();
 
 	/** Called when this state becomes the active state */
@@ -415,6 +420,7 @@ struct FGameFeaturePluginState
 
 	/** The common properties that can be accessed by the states of the state machine */
 	FGameFeaturePluginStateMachineProperties& StateProperties;
+	EGameFeaturePluginState State;
 
 	void UpdateStateMachineDeferred(float Delay = 0.0f) const;
 	void UpdateStateMachineImmediate() const;
@@ -451,7 +457,7 @@ private:
 /** Base class for destination game feature plugin states */
 struct FDestinationGameFeaturePluginState : public FGameFeaturePluginState
 {
-	FDestinationGameFeaturePluginState(FGameFeaturePluginStateMachineProperties& InStateProperties) : FGameFeaturePluginState(InStateProperties) {}
+	using FGameFeaturePluginState::FGameFeaturePluginState;
 
 	/** Returns the type of state this is */
 	virtual EGameFeaturePluginStateType GetStateType() const override { return EGameFeaturePluginStateType::Destination; }
@@ -463,7 +469,7 @@ struct FDestinationGameFeaturePluginState : public FGameFeaturePluginState
 /** Base class for error game feature plugin states */
 struct FErrorGameFeaturePluginState : public FDestinationGameFeaturePluginState
 {
-	FErrorGameFeaturePluginState(FGameFeaturePluginStateMachineProperties& InStateProperties) : FDestinationGameFeaturePluginState(InStateProperties) {}
+	using FDestinationGameFeaturePluginState::FDestinationGameFeaturePluginState;
 
 	/** Returns the type of state this is */
 	virtual EGameFeaturePluginStateType GetStateType() const override { return EGameFeaturePluginStateType::Error; }
@@ -491,137 +497,6 @@ class UGameFeaturePluginStateMachine : public UObject
 public:
 	UGameFeaturePluginStateMachine(const FObjectInitializer& ObjectInitializer);
 
-	/** Initializes the state machine and assigns the URL for the plugin it manages. This sets the machine to the 'UnknownStatus' state. */
-	void InitStateMachine(FGameFeaturePluginIdentifier InPluginIdentifier, const FGameFeatureProtocolOptions& InProtocolOptions);
-
-	/** Asynchronously transitions the state machine to the destination state range and reports when it is done. 
-	  * DestinationState must be of type EGameFeaturePluginStateType::Destination.
-	  * If returns true and OnFeatureStateTransitionComplete is not called immediately, OutCallbackHandle will be set
-	  * Returns false and does not callback if a transition is already in progress and the destination range is not compatible with the current range. */
-	bool SetDestination(FGameFeaturePluginStateRange InDestination, FGameFeatureStateTransitionComplete OnFeatureStateTransitionComplete, FDelegateHandle* OutCallbackHandle = nullptr);
-
-	/** Cancel the current transition if possible */
-	bool TryCancel(FGameFeatureStateTransitionCanceled OnFeatureStateTransitionCanceled, FDelegateHandle* OutCallbackHandle = nullptr);
-
-	/** Update the current PluginURL data for this plugin if possible. Returns false if this update fails or
-		if the supplied InPluginURL matches the existing URL data for this plugin.**/
-	UE::GameFeatures::FResult TryUpdatePluginProtocolOptions(const FGameFeatureProtocolOptions& InOptions, bool& bOutDidUpdate);
-
-	/** Remove any pending callback from SetDestination */
-	void RemovePendingTransitionCallback(FDelegateHandle InHandle);
-
-	/** Remove any pending callback from SetDestination */
-	void RemovePendingTransitionCallback(FDelegateUserObject DelegateObject);
-
-	/** Remove any pending callback from TryCancel */
-	void RemovePendingCancelCallback(FDelegateHandle InHandle);
-
-	/** Remove any pending callback from TryCancel */
-	void RemovePendingCancelCallback(FDelegateUserObject DelegateObject);
-
-	/** Returns the PluginIdentifier used to identify this plugin uniquely to the GameFeaturePlugin Subsystem */
-	const FGameFeaturePluginIdentifier& GetPluginIdentifier() const;
-
-	/** Returns the name of the game feature. Before StatusKnown, this returns the URL. */
-	const FString& GetGameFeatureName() const;
-
-	/** Returns the URL */
-	const FString& GetPluginURL() const;
-
-	/** Returns protocol medata */
-	const FGameFeatureProtocolMetadata& GetProtocolMetadata() const;
-
-	/** Returns any protocol options */
-	const FGameFeatureProtocolOptions& GetProtocolOptions() const;
-
-	/** Returns protocol options suitable for reuse by another state machine */
-	FGameFeatureProtocolOptions RecycleProtocolOptions() const;
-
-	/** Returns the plugin name if known (plugin must have been registered to know the name). */
-	const FString& GetPluginName() const;
-
-	/** Returns the uplugin filename of the game feature. Before StatusKnown, this returns false. */
-	bool GetPluginFilename(FString& OutPluginFilename) const;
-
-	/** Returns the enum state for this machine */
-	EGameFeaturePluginState GetCurrentState() const;
-
-	/** Returns the state range this machine is trying to move to */
-	FGameFeaturePluginStateRange GetDestination() const;
-
-	/** Returns information about the current state */
-	const FGameFeaturePluginStateInfo& GetCurrentStateInfo() const;
-
-	/** Returns true if attempting to reach a new destination state */
-	bool IsRunning() const;
-
-	/** Returns true if the state is at least StatusKnown so we can query info about the game feature plugin */
-	bool IsStatusKnown() const;
-
-	/** Returns true if the plugin is available to download/load. Only call if IsStatusKnown is true */
-	bool IsAvailable() const;
-
-	/** Returns true if the plugin is in an error state */
-	bool IsInErrorState() const;
-
-	/** Whether this machine is allowed to be asynchronous */
-	bool AllowAsyncLoading() const;
-
-	/** Returns true if the plugin will stream dependencies after installing. Only call if IsStatusKnown is true */
-	bool HasAssetStreamingDependencies() const;
-
-	void SetWasLoadedAsBuiltIn();
-
-	bool WasLoadedAsBuiltIn() const;
-
-	/** If the plugin is activated already, we will retrieve its game feature data */
-	UGameFeatureData* GetGameFeatureDataForActivePlugin();
-
-	/** If the plugin is registered already, we will retrieve its game feature data */
-	UGameFeatureData* GetGameFeatureDataForRegisteredPlugin(bool bCheckForRegistering = false);
-
-	const FGameFeaturePluginStateMachineProperties& GetProperties() const;
-
-	/** Returns true if the state machine is in an unrecoverable error state. */
-	bool IsErrorStateUnrecoverable() const;
-
-	/** Sets that the state machine is in an unrecoverable error. Used when the GFD is missing as that can't be recovered from. */
-	void SetUnrecoverableError();
-
-	/** Request to exclude this plugin from batch processing, for e.g. if the plugin is a dependency and has batch processing enabled */
-	void ExcludeFromBatchProcessing();
-
-	/** Helper static function to trigger state specific batch processing functions */
-	static void BatchProcess(EGameFeaturePluginState State, TConstArrayView<UGameFeaturePluginStateMachine*> GFPSMs);	
-
-private:
-	/** Returns true if the specified state is not a transition state */
-	bool IsValidTransitionState(EGameFeaturePluginState InState) const;
-
-	/** Returns true if the specified state is a destination state */
-	bool IsValidDestinationState(EGameFeaturePluginState InDestinationState) const;
-
-	/** Returns true if the specified state is a error state */
-	bool IsValidErrorState(EGameFeaturePluginState InDestinationState) const;
-
-	/** Processes the current state and looks for state transitions */
-	void UpdateStateMachine();
-
-	/** Update Progress for current state */
-	void UpdateCurrentStateProgress(float Progress);
-
-	/** Helper to register the state machine with the manager as a state machine in transiton */
-	void RegisterAsTransitioningStateMachine();
-
-	/** Helper to unregister the state machine with the manager as a state machine in transiton */
-	void UnregisterAsTransitioningStateMachine();
-	
-	/** Adds a batching request for current state if state can be batched */
-	void CheckAddBatchingRequestForCurrentState();
-
-	/** Cancels batching re1quest for current state if one exists */
-	void CheckAndCancelBatchingRequestForCurrentState();
-
 	/** Information about the current state */
 	FGameFeaturePluginStateInfo CurrentStateInfo;
 
diff --git a/Engine/Plugins/Runtime/GameFeatures/Source/GameFeatures/Private/GameFeaturesSubsystem.cpp b/Engine/Plugins/Runtime/GameFeatures/Source/GameFeatures/Private/GameFeaturesSubsystem.cpp
--- a/Engine/Plugins/Runtime/GameFeatures/Source/GameFeatures/Private/GameFeaturesSubsystem.cpp
+++ b/Engine/Plugins/Runtime/GameFeatures/Source/GameFeatures/Private/GameFeaturesSubsystem.cpp
@@ -851,7 +851,7 @@ void UGameFeaturesSubsystem::ForEachGameFeature(TFunctionRef<void(FGameFeatureIn
 	{
 		if (UGameFeaturePluginStateMachine* GFSM = StateMachineIt.Value())
 		{
-			FGameFeatureInfo GameFeatureInfo = { GFSM->GetPluginName(), GFSM->GetPluginURL(), GFSM->WasLoadedAsBuiltIn(), GFSM->GetCurrentState() };
+			FGameFeatureInfo GameFeatureInfo = { GetPluginName(GFSM), GetPluginURL(GFSM), WasLoadedAsBuiltIn(GFSM), GetCurrentState(GFSM) };
 			Visitor(MoveTemp(GameFeatureInfo));
 		}
 	}
@@ -880,19 +880,19 @@ void UGameFeaturesSubsystem::AddObserver(UObject* Observer, EObserverPluginState
 			{
 				if (UGameFeaturePluginStateMachine* GFSM = StateMachineIt.Value())
 				{
-					if (const UGameFeatureData* GameFeatureData = GFSM->GetGameFeatureDataForRegisteredPlugin(false))
+					if (const UGameFeatureData* GameFeatureData = GetGameFeatureDataForRegisteredPlugin(GFSM, false))
 					{
-						const FString& PluginName = GFSM->GetPluginName();
-						const FString& PluginURL = GFSM->GetPluginURL();
+						const FString& PluginName = GetPluginName(GFSM);
+						const FString& PluginURL = GetPluginURL(GFSM);
 
 						Interface->OnGameFeatureRegistering(GameFeatureData, PluginName, PluginURL);
 
-						if (GFSM->GetCurrentState() >= EGameFeaturePluginState::Loaded)
+						if (GetCurrentState(GFSM) >= EGameFeaturePluginState::Loaded)
 						{
 							Interface->OnGameFeatureLoading(GameFeatureData, PluginURL);
 						}
 
-						if (GFSM->GetCurrentState() >= EGameFeaturePluginState::Active)
+						if (GetCurrentState(GFSM) >= EGameFeaturePluginState::Active)
 						{
 							Interface->OnGameFeatureActivating(GameFeatureData, PluginURL);
 							Interface->OnGameFeatureActivated(GameFeatureData, PluginURL);
@@ -1398,12 +1398,12 @@ void UGameFeaturesSubsystem::OnGameFeaturePauseChange(const FGameFeaturePluginId
 
 const UGameFeatureData* UGameFeaturesSubsystem::GetDataForStateMachine(UGameFeaturePluginStateMachine* GFSM) const
 {
-	return GFSM->GetGameFeatureDataForActivePlugin();
+	return GetGameFeatureDataForActivePlugin(GFSM);
 }
 
 const UGameFeatureData* UGameFeaturesSubsystem::GetRegisteredDataForStateMachine(UGameFeaturePluginStateMachine* GFSM) const
 {
-	return GFSM->GetGameFeatureDataForRegisteredPlugin();
+	return GetGameFeatureDataForRegisteredPlugin(GFSM);
 }
 
 void UGameFeaturesSubsystem::GetGameFeatureDataForActivePlugins(TArray<const UGameFeatureData*>& OutActivePluginFeatureDatas)
@@ -1412,7 +1412,7 @@ void UGameFeaturesSubsystem::GetGameFeatureDataForActivePlugins(TArray<const UGa
 	{
 		if (UGameFeaturePluginStateMachine* GFSM = StateMachineIt.Value())
 		{
-			if (const UGameFeatureData* GameFeatureData = GFSM->GetGameFeatureDataForActivePlugin())
+			if (const UGameFeatureData* GameFeatureData = GetGameFeatureDataForActivePlugin(GFSM))
 			{
 				OutActivePluginFeatureDatas.Add(GameFeatureData);
 			}
@@ -1424,7 +1424,7 @@ const UGameFeatureData* UGameFeaturesSubsystem::GetGameFeatureDataForActivePlugi
 {
 	if (UGameFeaturePluginStateMachine* GFSM = FindGameFeaturePluginStateMachine(PluginURL))
 	{
-		return GFSM->GetGameFeatureDataForActivePlugin();
+		return GetGameFeatureDataForActivePlugin(GFSM);
 	}
 
 	return nullptr;
@@ -1434,7 +1434,7 @@ const UGameFeatureData* UGameFeaturesSubsystem::GetGameFeatureDataForRegisteredP
 {
 	if (UGameFeaturePluginStateMachine* GFSM = FindGameFeaturePluginStateMachine(PluginURL))
 	{
-		return GFSM->GetGameFeatureDataForRegisteredPlugin(bCheckForRegistering);
+		return GetGameFeatureDataForRegisteredPlugin(GFSM, bCheckForRegistering);
 	}
 
 	return nullptr;
@@ -1444,7 +1444,7 @@ bool UGameFeaturesSubsystem::IsGameFeaturePluginInstalled(const FString& PluginU
 {
 	if (const UGameFeaturePluginStateMachine* StateMachine = FindGameFeaturePluginStateMachine(PluginURL))
 	{
-		return StateMachine->GetCurrentState() >= EGameFeaturePluginState::Installed;
+		return GetCurrentState(StateMachine) >= EGameFeaturePluginState::Installed;
 	}
 	return false;
 }
@@ -1453,7 +1453,7 @@ bool UGameFeaturesSubsystem::IsGameFeaturePluginMounted(const FString& PluginURL
 {
 	if (const UGameFeaturePluginStateMachine* StateMachine = FindGameFeaturePluginStateMachine(PluginURL))
 	{
-		return StateMachine->GetCurrentState() > EGameFeaturePluginState::Mounting;
+		return GetCurrentState(StateMachine) > EGameFeaturePluginState::Mounting;
 	}
 	return false;
 }
@@ -1462,9 +1462,9 @@ bool UGameFeaturesSubsystem::IsGameFeaturePluginRegistered(const FString& Plugin
 {
 	if (const UGameFeaturePluginStateMachine* StateMachine = FindGameFeaturePluginStateMachine(PluginURL))
 	{
-		const EGameFeaturePluginState CurrentState = StateMachine->GetCurrentState();
+		const EGameFeaturePluginState CurrentState = GetCurrentState(StateMachine);
 
-		return StateMachine->GetCurrentState() >= EGameFeaturePluginState::Registered || (bCheckForRegistering && CurrentState == EGameFeaturePluginState::Registering);
+		return GetCurrentState(StateMachine) >= EGameFeaturePluginState::Registered || (bCheckForRegistering && CurrentState == EGameFeaturePluginState::Registering);
 	}
 	return false;
 }
@@ -1473,7 +1473,7 @@ bool UGameFeaturesSubsystem::IsGameFeaturePluginLoaded(const FString& PluginURL)
 {
 	if (const UGameFeaturePluginStateMachine* StateMachine = FindGameFeaturePluginStateMachine(PluginURL))
 	{
-		return StateMachine->GetCurrentState() >= EGameFeaturePluginState::Loaded;
+		return GetCurrentState(StateMachine) >= EGameFeaturePluginState::Loaded;
 	}
 	return false;
 }
@@ -1482,7 +1482,7 @@ bool UGameFeaturesSubsystem::WasGameFeaturePluginLoadedAsBuiltIn(const FString&
 {
 	if (const UGameFeaturePluginStateMachine* StateMachine = FindGameFeaturePluginStateMachine(PluginURL))
 	{
-		return StateMachine->WasLoadedAsBuiltIn();
+		return WasLoadedAsBuiltIn(StateMachine);
 	}
 	return false;
 }
@@ -1503,14 +1503,14 @@ void UGameFeaturesSubsystem::LoadGameFeaturePlugin(const FString& PluginURL, con
 
 	UGameFeaturePluginStateMachine* StateMachine = FindOrCreateGameFeaturePluginStateMachine(PluginURL, ProtocolOptions);
 
-	if (!StateMachine->IsRunning() && StateMachine->GetCurrentState() == EGameFeaturePluginState::Active)
+	if (!IsRunning(StateMachine) && GetCurrentState(StateMachine) == EGameFeaturePluginState::Active)
 	{
 		// TODO: Resolve the activated case here, this is needed because in a PIE environment the plugins
 		// are not sandboxed, and we need to do simulate a successful activate call in order run GFP systems 
 		// on whichever Role runs second between client and server.
 
 		// Refire the observer for Activated and do nothing else.
-		CallbackObservers(EObserverCallback::Activating, StateMachine->GetPluginIdentifier(), &StateMachine->GetPluginName(), StateMachine->GetGameFeatureDataForActivePlugin());
+		CallbackObservers(EObserverCallback::Activating, GetPluginIdentifier(StateMachine), &GetPluginName(StateMachine), GetGameFeatureDataForActivePlugin(StateMachine));
 	}
 
 	if (ShouldUpdatePluginProtocolOptions(StateMachine, ProtocolOptions))
@@ -1599,14 +1599,14 @@ void UGameFeaturesSubsystem::RegisterGameFeaturePlugin(const FString& PluginURL,
 
 	UGameFeaturePluginStateMachine* StateMachine = FindOrCreateGameFeaturePluginStateMachine(PluginURL, ProtocolOptions);
 
-	if (!StateMachine->IsRunning() && StateMachine->GetCurrentState() == EGameFeaturePluginState::Active)
+	if (!IsRunning(StateMachine) && GetCurrentState(StateMachine) == EGameFeaturePluginState::Active)
 	{
 		// TODO: Resolve the activated case here, this is needed because in a PIE environment the plugins
 		// are not sandboxed, and we need to do simulate a successful activate call in order run GFP systems 
 		// on whichever Role runs second between client and server.
 
 		// Refire the observer for Activated and do nothing else.
-		CallbackObservers(EObserverCallback::Activating, StateMachine->GetPluginIdentifier(), &StateMachine->GetPluginName(), StateMachine->GetGameFeatureDataForActivePlugin());
+		CallbackObservers(EObserverCallback::Activating, GetPluginIdentifier(StateMachine), &GetPluginName(StateMachine), GetGameFeatureDataForActivePlugin(StateMachine));
 	}
 
 	if (ShouldUpdatePluginProtocolOptions(StateMachine, ProtocolOptions))
@@ -1740,7 +1740,7 @@ void UGameFeaturesSubsystem::ChangeGameFeatureTargetState(const FString& PluginU
 
 	if (!bIsPluginAllowed)
 	{
-		if (TargetPluginState > StateMachine->GetCurrentState() || TargetPluginState > StateMachine->GetDestination())
+		if (TargetPluginState > GetCurrentState(StateMachine) || TargetPluginState > GetDestination(StateMachine))
 		{
 			UE_LOG(LogGameFeatures, Log, TEXT("Cannot change game feature target state: Plugin not allowed %s"), *PluginURL);
 
@@ -1751,15 +1751,15 @@ void UGameFeaturesSubsystem::ChangeGameFeatureTargetState(const FString& PluginU
 	}
 
 	if (TargetState == EGameFeatureTargetState::Active &&
-		!StateMachine->IsRunning() && 
-		StateMachine->GetCurrentState() == TargetPluginState)
+		!IsRunning(StateMachine) && 
+		GetCurrentState(StateMachine) == TargetPluginState)
 	{
 		// TODO: Resolve the activated case here, this is needed because in a PIE environment the plugins
 		// are not sandboxed, and we need to do simulate a successful activate call in order run GFP systems 
 		// on whichever Role runs second between client and server.
 
 		// Refire the observer for Activated and do nothing else.
-		CallbackObservers(EObserverCallback::Activating, StateMachine->GetPluginIdentifier(), &StateMachine->GetPluginName(), StateMachine->GetGameFeatureDataForActivePlugin());
+		CallbackObservers(EObserverCallback::Activating, GetPluginIdentifier(StateMachine), &GetPluginName(StateMachine), GetGameFeatureDataForActivePlugin(StateMachine));
 	}
 	
 	if (ShouldUpdatePluginProtocolOptions(StateMachine, ProtocolOptions))
@@ -1849,7 +1849,7 @@ UE::GameFeatures::FResult UGameFeaturesSubsystem::UpdateGameFeatureProtocolOptio
 	}
 
 	bool bUpdated = false;
-	UE::GameFeatures::FResult Result = StateMachine->TryUpdatePluginProtocolOptions(NewOptions, bUpdated);
+	UE::GameFeatures::FResult Result = TryUpdatePluginProtocolOptions(StateMachine, NewOptions, bUpdated);
 	if (bOutDidUpdate)
 	{
 		*bOutDidUpdate = bUpdated;
@@ -1871,12 +1871,12 @@ bool UGameFeaturesSubsystem::ShouldUpdatePluginProtocolOptions(const UGameFeatur
 	}
 	
 	//Make sure our StateMachine isn't in terminal, don't want to update Terminal plugins
-	if (TerminalGameFeaturePluginStateMachines.Contains(StateMachine) || (StateMachine->GetCurrentState() == EGameFeaturePluginState::Terminal))
+	if (TerminalGameFeaturePluginStateMachines.Contains(StateMachine) || GetCurrentState(StateMachine) == EGameFeaturePluginState::Terminal)
 	{
 		return false;
 	}
 
-	if (StateMachine->GetProtocolOptions() == NewOptions)
+	if (GetProtocolOptions(StateMachine) == NewOptions)
 	{
 		return false;
 	}
@@ -1899,9 +1899,9 @@ bool UGameFeaturesSubsystem::GetGameFeaturePluginInstallPercent(const FString& P
 {
 	if (const UGameFeaturePluginStateMachine* StateMachine = FindGameFeaturePluginStateMachine(PluginURL))
 	{
-		if (StateMachine->IsStatusKnown() && StateMachine->IsAvailable())
+		if (IsStatusKnown(StateMachine) && IsAvailable(StateMachine))
 		{
-			const FGameFeaturePluginStateInfo& StateInfo = StateMachine->GetCurrentStateInfo();
+			const FGameFeaturePluginStateInfo& StateInfo = GetCurrentStateInfo(StateMachine);
 
 			float InstallProgress = 0.0f;
 			if (StateInfo.State == EGameFeaturePluginState::Downloading)
@@ -1913,7 +1913,7 @@ bool UGameFeaturesSubsystem::GetGameFeaturePluginInstallPercent(const FString& P
 				InstallProgress = 1.0f;
 			}
 
-			if (!StateMachine->HasAssetStreamingDependencies())
+			if (!HasAssetStreamingDependencies(StateMachine))
 			{
 				Install_Percent = InstallProgress;
 				return true;
@@ -1990,7 +1990,7 @@ bool UGameFeaturesSubsystem::IsGameFeaturePluginActive(const FString& PluginURL,
 {
 	if (const UGameFeaturePluginStateMachine* StateMachine = FindGameFeaturePluginStateMachine(PluginURL))
 	{
-		const EGameFeaturePluginState CurrentState = StateMachine->GetCurrentState();
+		const EGameFeaturePluginState CurrentState = GetCurrentState(StateMachine);
 
 		return CurrentState == EGameFeaturePluginState::Active || (bCheckForActivating && CurrentState == EGameFeaturePluginState::Activating);
 	}
@@ -2287,10 +2287,10 @@ void UGameFeaturesSubsystem::UninstallGameFeaturePlugin(const FString& PluginURL
 	check(StateMachine);
 
 	// We may need to update our ProtocolOptions to force certain metadata changes to facilitate this uninstall
-	FGameFeatureProtocolOptions ProtocolOptions = StateMachine->GetProtocolOptions();
+	FGameFeatureProtocolOptions ProtocolOptions = GetProtocolOptions(StateMachine);
 
 	// InstallBundle Protocol GameFeatures may need to change their metadata to force this uninstall
-	if (StateMachine->GetPluginIdentifier().GetPluginProtocol() == EGameFeaturePluginProtocol::InstallBundle)
+	if (GetPluginIdentifier(StateMachine).GetPluginProtocol() == EGameFeaturePluginProtocol::InstallBundle)
 	{
 		// It's possible that ParseURL hasn't been called yet so setup options here if needed.
 		if (!ProtocolOptions.HasSubtype<FInstallBundlePluginProtocolOptions>())
@@ -2392,7 +2392,7 @@ void UGameFeaturesSubsystem::CancelGameFeatureStateChange(const FString& PluginU
 {
 	if (UGameFeaturePluginStateMachine* StateMachine = FindGameFeaturePluginStateMachine(PluginURL))
 	{
-		const bool bCancelPending = StateMachine->TryCancel(FGameFeatureStateTransitionCanceled::CreateWeakLambda(this, [CompleteDelegate](UGameFeaturePluginStateMachine* Machine)
+		const bool bCancelPending = TryCancel(StateMachine, FGameFeatureStateTransitionCanceled::CreateWeakLambda(this, [CompleteDelegate](UGameFeaturePluginStateMachine* Machine)
 		{
 			CompleteDelegate.ExecuteIfBound(UE::GameFeatures::FResult(MakeValue()));
 		}));
@@ -2514,7 +2514,7 @@ void UGameFeaturesSubsystem::LoadBuiltInGameFeaturePlugin(const TSharedRef<IPlug
 			{
 				if (const UGameFeaturePluginStateMachine* GFSM = FindGameFeaturePluginStateMachine(PluginURL))
 				{
-					FGameFeatureProtocolOptions CurrentProtocolOptions = GFSM->GetProtocolOptions();
+					FGameFeatureProtocolOptions CurrentProtocolOptions = GetProtocolOptions(GFSM);
 					if (CurrentProtocolOptions.HasSubtype<FInstallBundlePluginProtocolOptions>())
 					{
 						const FInstallBundlePluginProtocolOptions& CurrentInstallBundleOptions = CurrentProtocolOptions.GetSubtype<FInstallBundlePluginProtocolOptions>();
@@ -2556,7 +2556,7 @@ void UGameFeaturesSubsystem::LoadBuiltInGameFeaturePlugin(const TSharedRef<IPlug
 			}
 			const EGameFeaturePluginState DestinationState = ConvertInitialFeatureStateToTargetState(InitialAutoState);
 
-			StateMachine->SetWasLoadedAsBuiltIn();
+			SetWasLoadedAsBuiltIn(StateMachine);
 
 			// If we're already at the destination or beyond, don't transition back
 			FGameFeaturePluginStateRange Destination(DestinationState, EGameFeaturePluginState::Active);
@@ -2850,7 +2850,7 @@ FString UGameFeaturesSubsystem::GetPluginFilenameFromPluginURL(const FString& Pl
 {
 	FString PluginFilename;
 	const UGameFeaturePluginStateMachine* GFSM = FindGameFeaturePluginStateMachine(PluginURL);
-	if (GFSM == nullptr || !GFSM->GetPluginFilename(PluginFilename))
+	if (GFSM == nullptr || !GetPluginFilename(GFSM, PluginFilename))
 	{
 		UE_LOG(LogGameFeatures, Error, TEXT("UGameFeaturesSubsystem could not get the plugin path from the plugin URL. URL:%s "), *PluginURL);
 	}
@@ -2865,7 +2865,7 @@ FString UGameFeaturesSubsystem::GetPluginNameFromPluginURL(const FString& Plugin
 		UE_LOG(LogGameFeatures, Error, TEXT("UGameFeaturesSubsystem could not get the plugin name from the plugin URL. URL:%s "), *PluginURL);
 		return FString();
 	}
-	return GFSM->GetPluginName();
+	return GetPluginName(GFSM);
 }
 
 void UGameFeaturesSubsystem::FixPluginPackagePath(FString& PathToFix, const FString& PluginRootPath, bool bMakeRelativeToPluginRoot)
@@ -2889,10 +2889,10 @@ void UGameFeaturesSubsystem::GetLoadedGameFeaturePluginFilenamesForCooking(TArra
 	for (auto StateMachineIt = GameFeaturePluginStateMachines.CreateConstIterator(); StateMachineIt; ++StateMachineIt)
 	{
 		UGameFeaturePluginStateMachine* GFSM = StateMachineIt.Value();
-		if (GFSM && GFSM->GetCurrentState() > EGameFeaturePluginState::Installed)
+		if (GFSM && GetCurrentState(GFSM) > EGameFeaturePluginState::Installed)
 		{
 			FString PluginFilename;
-			if (GFSM->GetPluginFilename(PluginFilename))
+			if (GetPluginFilename(GFSM, PluginFilename))
 			{
 				OutLoadedPluginFilenames.Add(PluginFilename);
 			}
@@ -2910,7 +2910,7 @@ EGameFeaturePluginState UGameFeaturesSubsystem::GetPluginState(FGameFeaturePlugi
 {
 	if (UGameFeaturePluginStateMachine* StateMachine = FindGameFeaturePluginStateMachine(PluginIdentifier))
 	{
-		return StateMachine->GetCurrentState();
+		return GetCurrentState(StateMachine);
 	}
 	else
 	{
@@ -2977,7 +2977,7 @@ bool UGameFeaturesSubsystem::GetGameFeatureControlsUPlugin(const FString& Plugin
 {
 	if (UGameFeaturePluginStateMachine* Machine = FindGameFeaturePluginStateMachine(PluginURL))
 	{
-		OutGameFeatureControlsUPlugin = Machine->GetProperties().bAddedPluginToManager;
+		OutGameFeatureControlsUPlugin = GetProperties(Machine).bAddedPluginToManager;
 		return true;
 	}
 
@@ -3350,10 +3350,10 @@ struct FGameFeaturePluginPredownloadContext : public FGameFeaturePluginPredownlo
 			FGFPData& GFPData = Pair.Value;
 
 			UGameFeaturePluginStateMachine* Machine = GFPSubSys.FindGameFeaturePluginStateMachine(GFPIdentifier);
-			if (Machine && Machine->GetDestination() < EGameFeaturePluginState::Installed)
+			if (Machine && GFPSubSys.GetDestination(Machine) < EGameFeaturePluginState::Installed)
 			{
 				// Existing machine exists and wants to be uninstalled, can't precache
-				UE_LOGFMT(LogGameFeatures, Error, "GFP Predownload failed because a GFP is unloading, GFP: {GFP}", ("GFP", Machine->GetPluginName()));
+				UE_LOGFMT(LogGameFeatures, Error, "GFP Predownload failed because a GFP is unloading, GFP: {GFP}", ("GFP", GFPSubSys.GetPluginName(Machine)));
 				UE::GameFeatures::FResult ErrorResult = MakeError(FString::Printf(TEXT("%.*s%s"),
 					PredownloadErrorNamespace.Len(), PredownloadErrorNamespace.GetData(),
 					TEXT("GFPUnloading")));
@@ -3746,10 +3746,10 @@ struct FGameFeaturePluginPredownloadContext : public FGameFeaturePluginPredownlo
 		{
 			UGameFeaturePluginStateMachine* Machine = Pair.Value;
 			if (Machine &&
-				Machine->GetCurrentState() > EGameFeaturePluginState::StatusKnown &&
-				Machine->GetCurrentState() != EGameFeaturePluginState::Releasing)
+				GFPSubSys.GetCurrentState(Machine) > EGameFeaturePluginState::StatusKnown &&
+				GFPSubSys.GetCurrentState(Machine) != EGameFeaturePluginState::Releasing)
 			{
-				const FGameFeatureProtocolMetadata& ProtocolMetaData = Machine->GetProtocolMetadata();
+				const FGameFeatureProtocolMetadata& ProtocolMetaData = GFPSubSys.GetProtocolMetadata(Machine);
 				if (ProtocolMetaData.HasSubtype<FInstallBundlePluginProtocolMetaData>())
 				{
 					const FInstallBundlePluginProtocolMetaData& ProtocolData = ProtocolMetaData.GetSubtype<FInstallBundlePluginProtocolMetaData>();
@@ -3765,8 +3765,8 @@ struct FGameFeaturePluginPredownloadContext : public FGameFeaturePluginPredownlo
 
 			UGameFeaturePluginStateMachine* Machine = GFPSubSys.FindGameFeaturePluginStateMachine(GFPIdentifier);
 			if (Machine &&
-				Machine->GetCurrentState() > EGameFeaturePluginState::StatusKnown &&
-				Machine->GetCurrentState() != EGameFeaturePluginState::Releasing)
+				GFPSubSys.GetCurrentState(Machine) > EGameFeaturePluginState::StatusKnown &&
+				GFPSubSys.GetCurrentState(Machine) != EGameFeaturePluginState::Releasing)
 			{
 				// A machine is using the bundles, don't release
 				KeepList.Append(GFPData.ProtocolMetadata.InstallBundles);
@@ -3807,7 +3807,7 @@ UGameFeaturePluginStateMachine* UGameFeaturesSubsystem::FindGameFeaturePluginSta
 		GameFeaturePluginStateMachines.FindByHash(GetTypeHash(PluginIdentifier.GetIdentifyingString()), PluginIdentifier.GetIdentifyingString());
 	if (ExistingStateMachine)
 	{
-		EGameFeaturePluginProtocol ExpectedProtocol = (*ExistingStateMachine)->GetPluginIdentifier().GetPluginProtocol();
+		EGameFeaturePluginProtocol ExpectedProtocol = GetPluginIdentifier(*ExistingStateMachine).GetPluginProtocol();
 		if (ensureMsgf(ExpectedProtocol == PluginIdentifier.GetPluginProtocol(), TEXT("Expected protocol %s for %.*s"), UE::GameFeatures::GameFeaturePluginProtocolPrefix(ExpectedProtocol), ShortUrl.Len(), ShortUrl.GetData()))
 		{
 			UE_LOG(LogGameFeatures, VeryVerbose, TEXT("FOUND GameFeaturePlugin using PluginIdentifier:%.*s for PluginURL:%s"), ShortUrl.Len(), ShortUrl.GetData(), *PluginIdentifier.GetFullPluginURL());
@@ -3837,7 +3837,7 @@ UGameFeaturePluginStateMachine* UGameFeaturesSubsystem::FindOrCreateGameFeatureP
 	{
 		// In this case, still return the existing machine, even if the protocol doesn't match. This function should never return null.
 		// There can only be one active instance of any machine.
-		EGameFeaturePluginProtocol ExpectedProtocol = (*ExistingStateMachine)->GetPluginIdentifier().GetPluginProtocol();
+		EGameFeaturePluginProtocol ExpectedProtocol = GetPluginIdentifier(*ExistingStateMachine).GetPluginProtocol();
 		ensureMsgf(ExpectedProtocol == PluginIdentifier.GetPluginProtocol(), TEXT("Expected protocol %s for %.*s"), UE::GameFeatures::GameFeaturePluginProtocolPrefix(ExpectedProtocol), PluginIdentifier.GetIdentifyingString().Len(), PluginIdentifier.GetIdentifyingString().GetData());
 
 		UE_LOG(LogGameFeatures, VeryVerbose, TEXT("Found GameFeaturePlugin StateMachine using Identifier:%.*s from PluginURL:%s"), PluginIdentifier.GetIdentifyingString().Len(), PluginIdentifier.GetIdentifyingString().GetData(), *PluginURL);
@@ -3856,7 +3856,7 @@ UGameFeaturePluginStateMachine* UGameFeaturesSubsystem::FindOrCreateGameFeatureP
 
 	UGameFeaturePluginStateMachine* NewStateMachine = NewObject<UGameFeaturePluginStateMachine>(this);
 	GameFeaturePluginStateMachines.Add(FString(PluginIdentifier.GetIdentifyingString()), NewStateMachine);
-	NewStateMachine->InitStateMachine(MoveTemp(PluginIdentifier), ProtocolOptions);
+	InitStateMachine(NewStateMachine, MoveTemp(PluginIdentifier), ProtocolOptions);
 
 	return NewStateMachine;
 }
@@ -3866,15 +3866,15 @@ void UGameFeaturesSubsystem::LoadBuiltInGameFeaturePluginComplete(const UE::Game
 	check(Machine);
 	if (Result.HasValue())
 	{
-		checkf(RequestedDestination.Contains(Machine->GetCurrentState()), TEXT("Game feature '%s': Ending state %s is not in expected range [%s, %s]"), 
-			*Machine->GetGameFeatureName(), 
-			*UE::GameFeatures::ToString(Machine->GetCurrentState()), 
+		checkf(RequestedDestination.Contains(GetCurrentState(Machine)), TEXT("Game feature '%s': Ending state %s is not in expected range [%s, %s]"), 
+			*GetGameFeatureName(Machine), 
+			*UE::GameFeatures::ToString(GetCurrentState(Machine)), 
 			*UE::GameFeatures::ToString(RequestedDestination.MinState), 
 			*UE::GameFeatures::ToString(RequestedDestination.MaxState));
 	}
 	else
 	{
-		SetExplanationForNotMountingPlugin(Machine->GetPluginURL(), UE::GameFeatures::ToString(Result));
+		SetExplanationForNotMountingPlugin(GetPluginURL(Machine), UE::GameFeatures::ToString(Result));
 	}
 }
 
@@ -3885,18 +3885,18 @@ void UGameFeaturesSubsystem::ChangeGameFeatureDestination(UGameFeaturePluginStat
 
 void UGameFeaturesSubsystem::ChangeGameFeatureDestination(UGameFeaturePluginStateMachine* Machine, const FGameFeatureProtocolOptions& InProtocolOptions, const FGameFeaturePluginStateRange& StateRange, FGameFeaturePluginChangeStateComplete CompleteDelegate)
 {
-	const bool bSetDestination = Machine->SetDestination(StateRange,
+	const bool bSetDestination = SetDestination(Machine, StateRange,
 		FGameFeatureStateTransitionComplete::CreateUObject(this, &ThisClass::ChangeGameFeatureTargetStateComplete, CompleteDelegate));
 
 	if (bSetDestination)
 	{
-		UE_LOG(LogGameFeatures, Verbose, TEXT("ChangeGameFeatureDestination: Set Game Feature %s Destination State to [%s, %s]"), *Machine->GetGameFeatureName(), *UE::GameFeatures::ToString(StateRange.MinState), *UE::GameFeatures::ToString(StateRange.MaxState));
+		UE_LOG(LogGameFeatures, Verbose, TEXT("ChangeGameFeatureDestination: Set Game Feature %s Destination State to [%s, %s]"), *GetGameFeatureName(Machine), *UE::GameFeatures::ToString(StateRange.MinState), *UE::GameFeatures::ToString(StateRange.MaxState));
 	}
 	else
 	{
-		FGameFeaturePluginStateRange CurrDestination = Machine->GetDestination();
+		FGameFeaturePluginStateRange CurrDestination = GetDestination(Machine);
 		UE_LOG(LogGameFeatures, Display, TEXT("ChangeGameFeatureDestination: Attempting to cancel transition for Game Feature %s. Desired [%s, %s]. Current [%s, %s]"), 
-			*Machine->GetGameFeatureName(), 
+			*GetGameFeatureName(Machine), 
 			*UE::GameFeatures::ToString(StateRange.MinState), *UE::GameFeatures::ToString(StateRange.MaxState),
 			*UE::GameFeatures::ToString(CurrDestination.MinState), *UE::GameFeatures::ToString(CurrDestination.MaxState));
 
@@ -3904,33 +3904,33 @@ void UGameFeaturesSubsystem::ChangeGameFeatureDestination(UGameFeaturePluginStat
 		auto OnCanceled = [this, InProtocolOptions, StateRange, CompleteDelegate](UGameFeaturePluginStateMachine* Machine) mutable
 		{
 			// Special case for terminal state since it cannot be exited, we need to make a new machine
-			if (Machine->GetCurrentState() == EGameFeaturePluginState::Terminal)
+			if (GetCurrentState(Machine) == EGameFeaturePluginState::Terminal)
 			{
-				UGameFeaturePluginStateMachine* NewMachine = FindOrCreateGameFeaturePluginStateMachine(Machine->GetPluginURL(), InProtocolOptions);
-				checkf(NewMachine != Machine, TEXT("Game Feature Plugin %s should have already been removed from subsystem!"), *Machine->GetPluginURL());
+				UGameFeaturePluginStateMachine* NewMachine = FindOrCreateGameFeaturePluginStateMachine(GetPluginURL(Machine), InProtocolOptions);
+				checkf(NewMachine != Machine, TEXT("Game Feature Plugin %s should have already been removed from subsystem!"), *GetPluginURL(Machine));
 				Machine = NewMachine;
 			}
 
 			// Now that the transition has been canceled, retry reaching the desired destination
-			const bool bSetDestination = Machine->SetDestination(StateRange,
+			const bool bSetDestination = SetDestination(Machine, StateRange,
 				FGameFeatureStateTransitionComplete::CreateUObject(this, &ThisClass::ChangeGameFeatureTargetStateComplete, CompleteDelegate));
 
 			if (!ensure(bSetDestination))
 			{
-				UE_LOG(LogGameFeatures, Warning, TEXT("ChangeGameFeatureDestination: Failed to set Game Feature %s Destination State to [%s, %s]"), *Machine->GetGameFeatureName(), *UE::GameFeatures::ToString(StateRange.MinState), *UE::GameFeatures::ToString(StateRange.MaxState));
+				UE_LOG(LogGameFeatures, Warning, TEXT("ChangeGameFeatureDestination: Failed to set Game Feature %s Destination State to [%s, %s]"), *GetGameFeatureName(Machine), *UE::GameFeatures::ToString(StateRange.MinState), *UE::GameFeatures::ToString(StateRange.MaxState));
 
 				CompleteDelegate.ExecuteIfBound(UE::GameFeatures::FResult(MakeError(UE::GameFeatures::SubsystemErrorNamespace + UE::GameFeatures::CommonErrorCodes::UnreachableState)));
 			}
 			else
 			{
-				UE_LOG(LogGameFeatures, Display, TEXT("ChangeGameFeatureDestination: OnCanceled, set Game Feature %s Destination State to [%s, %s]"), *Machine->GetGameFeatureName(), *UE::GameFeatures::ToString(StateRange.MinState), *UE::GameFeatures::ToString(StateRange.MaxState));
+				UE_LOG(LogGameFeatures, Display, TEXT("ChangeGameFeatureDestination: OnCanceled, set Game Feature %s Destination State to [%s, %s]"), *GetGameFeatureName(Machine), *UE::GameFeatures::ToString(StateRange.MinState), *UE::GameFeatures::ToString(StateRange.MaxState));
 			}
 		};
 
-		const bool bCancelPending = Machine->TryCancel(FGameFeatureStateTransitionCanceled::CreateWeakLambda(this, MoveTemp(OnCanceled)));
+		const bool bCancelPending = TryCancel(Machine, FGameFeatureStateTransitionCanceled::CreateWeakLambda(this, MoveTemp(OnCanceled)));
 		if (!ensure(bCancelPending))
 		{
-			UE_LOG(LogGameFeatures, Warning, TEXT("ChangeGameFeatureDestination: Failed to cancel Game Feature %s"), *Machine->GetGameFeatureName());
+			UE_LOG(LogGameFeatures, Warning, TEXT("ChangeGameFeatureDestination: Failed to cancel Game Feature %s"), *GetGameFeatureName(Machine));
 
 			CompleteDelegate.ExecuteIfBound(UE::GameFeatures::FResult(MakeError(UE::GameFeatures::SubsystemErrorNamespace + UE::GameFeatures::CommonErrorCodes::UnreachableState + UE::GameFeatures::CommonErrorCodes::CancelAddonCode)));
 		}
@@ -3940,9 +3940,9 @@ void UGameFeaturesSubsystem::ChangeGameFeatureDestination(UGameFeaturePluginStat
 void UGameFeaturesSubsystem::ChangeGameFeatureTargetStateComplete(UGameFeaturePluginStateMachine* Machine, const UE::GameFeatures::FResult& Result, FGameFeaturePluginChangeStateComplete CompleteDelegate)
 {
 #if WITH_EDITOR
-	if (!Result.HasError() && Machine->GetCurrentState() > EGameFeaturePluginState::Mounting)
+	if (!Result.HasError() && GetCurrentState(Machine) > EGameFeaturePluginState::Mounting)
 	{
-		UnmountedPluginNameToExplanation.Remove(Machine->GetPluginName());
+		UnmountedPluginNameToExplanation.Remove(GetPluginName(Machine));
 	}
 #endif
 	CompleteDelegate.ExecuteIfBound(Result);
@@ -3951,18 +3951,18 @@ void UGameFeaturesSubsystem::ChangeGameFeatureTargetStateComplete(UGameFeaturePl
 void UGameFeaturesSubsystem::BeginTermination(UGameFeaturePluginStateMachine* Machine)
 {
 	check(IsValid(Machine));
-	check(Machine->GetCurrentState() == EGameFeaturePluginState::Terminal);
+	check(GetCurrentState(Machine) == EGameFeaturePluginState::Terminal);
 
-	FStringView Identifer = Machine->GetPluginIdentifier().GetIdentifyingString();
+	FStringView Identifer = GetPluginIdentifier(Machine).GetIdentifyingString();
 
-	UE_LOG(LogGameFeatures, Verbose, TEXT("BeginTermination of GameFeaturePlugin. Identifier:%.*s URL:%s"), Identifer.Len(), Identifer.GetData(), *(Machine->GetPluginURL()));
+	UE_LOG(LogGameFeatures, Verbose, TEXT("BeginTermination of GameFeaturePlugin. Identifier:%.*s URL:%s"), Identifer.Len(), Identifer.GetData(), *GetPluginURL(Machine));
 	GameFeaturePluginStateMachines.RemoveByHash(GetTypeHash(Identifer), Identifer);
 	TerminalGameFeaturePluginStateMachines.Add(Machine);
 }
 
 void UGameFeaturesSubsystem::FinishTermination(UGameFeaturePluginStateMachine* Machine)
 {
-	UE_LOG(LogGameFeatures, Verbose, TEXT("FinishTermination of GameFeaturePlugin. Identifier:%.*s URL:%s"), Machine->GetPluginIdentifier().GetIdentifyingString().Len(), Machine->GetPluginIdentifier().GetIdentifyingString().GetData(), *(Machine->GetPluginURL()));
+	UE_LOG(LogGameFeatures, Verbose, TEXT("FinishTermination of GameFeaturePlugin. Identifier:%.*s URL:%s"), GetPluginIdentifier(Machine).GetIdentifyingString().Len(), GetPluginIdentifier(Machine).GetIdentifyingString().GetData(), *GetPluginURL(Machine));
 	TerminalGameFeaturePluginStateMachines.RemoveSwap(Machine);
 }
 
@@ -4043,7 +4043,7 @@ bool UGameFeaturesSubsystem::FindOrCreatePluginDependencyStateMachines(const FSt
 				// Propogate bWasLoadedAsBuiltInGameFeaturePlugin
 				if (InStateProperties.bWasLoadedAsBuiltInGameFeaturePlugin)
 				{
-					ResolvedDependency->SetWasLoadedAsBuiltIn();
+					SetWasLoadedAsBuiltIn(ResolvedDependency);
 				}
 
 				// Note: Given that LoadBuiltInGameFeaturePlugins does a topological sort, we don't expect to hit this path for built-ins
@@ -4052,12 +4052,12 @@ bool UGameFeaturesSubsystem::FindOrCreatePluginDependencyStateMachines(const FSt
 					if (InStateProperties.bWasLoadedAsBuiltInGameFeaturePlugin)
 					{
 						UE_LOGFMT(LogGameFeatures, Warning, "GFP dependency {Dep} was forcibly created by {Parent}, Game specific policies may be incorrectly filtering this dependency.",
-							("Dep", ResolvedDependency->GetPluginIdentifier().GetIdentifyingString()), ("Parent", InStateProperties.PluginIdentifier.GetIdentifyingString()));
+							("Dep", GetPluginIdentifier(ResolvedDependency).GetIdentifyingString()), ("Parent", InStateProperties.PluginIdentifier.GetIdentifyingString()));
 					}
 					else
 					{
 						UE_LOGFMT(LogGameFeatures, Warning, "GFP dependency {Dep} was unexpectedly forcibly created by {Parent}",
-							("Dep", ResolvedDependency->GetPluginIdentifier().GetIdentifyingString()), ("Parent", InStateProperties.PluginIdentifier.GetIdentifyingString()));
+							("Dep", GetPluginIdentifier(ResolvedDependency).GetIdentifyingString()), ("Parent", InStateProperties.PluginIdentifier.GetIdentifyingString()));
 					}
 				}
 				else if (bErrorOnDepCreation)
@@ -4065,12 +4065,12 @@ bool UGameFeaturesSubsystem::FindOrCreatePluginDependencyStateMachines(const FSt
 					if (InStateProperties.bWasLoadedAsBuiltInGameFeaturePlugin)
 					{
 						UE_LOGFMT(LogGameFeatures, Error, "GFP dependency {Dep} was forcibly created by {Parent}, Game specific policies may be incorrectly filtering this dependency.",
-							("Dep", ResolvedDependency->GetPluginIdentifier().GetIdentifyingString()), ("Parent", InStateProperties.PluginIdentifier.GetIdentifyingString()));
+							("Dep", GetPluginIdentifier(ResolvedDependency).GetIdentifyingString()), ("Parent", InStateProperties.PluginIdentifier.GetIdentifyingString()));
 					}
 					else
 					{
 						UE_LOGFMT(LogGameFeatures, Error, "GFP dependency {Dep} was unexpectedly forcibly created by {Parent}",
-							("Dep", ResolvedDependency->GetPluginIdentifier().GetIdentifyingString()), ("Parent", InStateProperties.PluginIdentifier.GetIdentifyingString()));
+							("Dep", GetPluginIdentifier(ResolvedDependency).GetIdentifyingString()), ("Parent", InStateProperties.PluginIdentifier.GetIdentifyingString()));
 					}
 				}
 			}
@@ -4187,7 +4187,7 @@ void UGameFeaturesSubsystem::ListGameFeaturePlugins(const TArray<FString>& Args,
 	GameFeaturePluginStateMachines.GenerateValueArray(StateMachines);
 
 	// Alphasort
-	StateMachines.Sort([](const UGameFeaturePluginStateMachine& A, const UGameFeaturePluginStateMachine& B) { return A.GetGameFeatureName().Compare(B.GetGameFeatureName()) < 0; });
+	StateMachines.Sort([this](const UGameFeaturePluginStateMachine& A, const UGameFeaturePluginStateMachine& B) { return GetGameFeatureName(&A).Compare(GetGameFeatureName(&B)) < 0; });
 
 	int32 PluginCount = 0;
 	for (UGameFeaturePluginStateMachine* GFSM : StateMachines)
@@ -4197,18 +4197,18 @@ void UGameFeaturesSubsystem::ListGameFeaturePlugins(const TArray<FString>& Args,
 			continue;
 		}
 
-		if (bActiveOnly && GFSM->GetCurrentState() != EGameFeaturePluginState::Active)
+		if (bActiveOnly && GetCurrentState(GFSM) != EGameFeaturePluginState::Active)
 		{
 			continue;
 		}
 
 		if (bCsv)
 		{
-			Ar.Logf(TEXT(",%s,%s"), *GFSM->GetGameFeatureName(), *UE::GameFeatures::ToString(GFSM->GetCurrentState()));
+			Ar.Logf(TEXT(",%s,%s"), *GetGameFeatureName(GFSM), *UE::GameFeatures::ToString(GetCurrentState(GFSM)));
 		}
 		else
 		{
-			Ar.Logf(TEXT("%s (%s)"), *GFSM->GetGameFeatureName(), *UE::GameFeatures::ToString(GFSM->GetCurrentState()));
+			Ar.Logf(TEXT("%s (%s)"), *GetGameFeatureName(GFSM), *UE::GameFeatures::ToString(GetCurrentState(GFSM)));
 		}
 		++PluginCount;
 	}
@@ -4518,11 +4518,11 @@ bool UGameFeaturesSubsystem::TickBatchProcessing()
 	bool bAreAnyGFPSMsPendingAtFence = false;
 	for (UGameFeaturePluginStateMachine* GFPSM : RunningStateMachines)
 	{
-		if (GFPSM->GetProperties().IsWaitingForBatchProcessing() && GFPSM->GetCurrentState() == CurrentFence)
+		if (GetProperties(GFPSM).IsWaitingForBatchProcessing() && GetCurrentState(GFPSM) == CurrentFence)
 		{
 			AtFence.Add(GFPSM);
 		}
-		if (GFPSM->GetDestination().MinState >= CurrentFence && GFPSM->GetCurrentState() < CurrentFence)
+		if (GetDestination(GFPSM).MinState >= CurrentFence && GetCurrentState(GFPSM) < CurrentFence)
 		{
 			bAreAnyGFPSMsPendingAtFence = true;
 			break;
@@ -4535,7 +4535,7 @@ bool UGameFeaturesSubsystem::TickBatchProcessing()
 
 		FGameFeatureBatchProcessingFence& Fence = BatchProcessingFences[CurrentFence];
 
-		UGameFeaturePluginStateMachine::BatchProcess(CurrentFence, AtFence);
+		BatchProcess(CurrentFence, AtFence);
 
 		Fence.NotifyUpdateStateMachines.Broadcast();
 
@@ -4554,10 +4554,10 @@ TSet<FString> UGameFeaturesSubsystem::GetActivePluginNames() const
 	for (const TPair<FString, TObjectPtr<UGameFeaturePluginStateMachine>>& Pair : GameFeaturePluginStateMachines)
 	{
 		UGameFeaturePluginStateMachine* StateMachine = Pair.Value;
-		if (StateMachine->GetCurrentState() == EGameFeaturePluginState::Active &&
-			StateMachine->GetDestination().Contains(EGameFeaturePluginState::Active))
+		if (GetCurrentState(StateMachine) == EGameFeaturePluginState::Active &&
+			GetDestination(StateMachine).Contains(EGameFeaturePluginState::Active))
 		{
-			ActivePluginNames.Add(StateMachine->GetPluginName());
+			ActivePluginNames.Add(GetPluginName(StateMachine));
 		}
 	}
 
@@ -4737,7 +4737,7 @@ const TSet<FString>& UGameFeaturesSubsystem::GetDlcPlugins()
 	return DlcPlugins;
 }
 
-bool UGameFeaturesSubsystem::GetPluginDebugStateEnabled(const FString& PluginUrl)
+bool UGameFeaturesSubsystem::GetPluginDebugStateEnabled(const FString& PluginUrl) const
 {
 #if !UE_BUILD_SHIPPING
 	return DebugStateChangedForPlugins.Contains(PluginUrl);
diff --git a/Engine/Plugins/Runtime/GameFeatures/Source/GameFeatures/Public/GameFeaturesSubsystem.h b/Engine/Plugins/Runtime/GameFeatures/Source/GameFeatures/Public/GameFeaturesSubsystem.h
--- a/Engine/Plugins/Runtime/GameFeatures/Source/GameFeatures/Public/GameFeaturesSubsystem.h
+++ b/Engine/Plugins/Runtime/GameFeatures/Source/GameFeatures/Public/GameFeaturesSubsystem.h
@@ -27,8 +27,10 @@ class UGameFeaturesProjectPolicies;
 class IPlugin;
 class FJsonObject;
 struct FWorldContext;
+struct FGameFeaturePluginStateInfo;
 struct FGameFeaturePluginStateRange;
 struct FGameFeaturePluginStateMachineProperties;
+struct FGameFeatureProtocolMetadata;
 struct FInstallBundleReleaseRequestInfo;
 enum class EInstallBundleResult : uint32;
 enum class EInstallBundleRequestFlags : uint32;
@@ -767,7 +769,7 @@ public:
 
 	static UE_API const TSet<FString>& GetDlcPlugins();
 
-	UE_API bool GetPluginDebugStateEnabled(const FString& PluginUrl);
+	UE_API bool GetPluginDebugStateEnabled(const FString& PluginUrl) const;
 	UE_API void SetPluginDebugStateEnabled(const FString& PluginUrl, bool bEnabled);
 
 	/**
@@ -1016,6 +1018,153 @@ private:
 	bool bInitializedPolicyManager = false;
 
 	static UE_API TSet<FString> DlcPlugins;
+
+
+private:
+	/**
+	 * State machine execution functions hoisted from UGameFeaturePluginStateMachine etc as a refactor
+	 * to enable more batch & parallel processing.
+	 * These fields could be further factored out into another class.
+	 * These methods are marked const until we make changes that hoist state out into the subsystem.
+	 */
+
+	/** Initializes the state machine and assigns the URL for the plugin it manages. This sets the machine to the 'UnknownStatus' state. */
+	void InitStateMachine(UGameFeaturePluginStateMachine* This, FGameFeaturePluginIdentifier InPluginIdentifier, const FGameFeatureProtocolOptions& InProtocolOptions);
+
+	/** Asynchronously transitions the state machine to the destination state range and reports when it is done. 
+	  * DestinationState must be of type EGameFeaturePluginStateType::Destination.
+	  * If returns true and OnFeatureStateTransitionComplete is not called immediately, OutCallbackHandle will be set
+	  * Returns false and does not callback if a transition is already in progress and the destination range is not compatible with the current range. */
+	bool SetDestination(
+		UGameFeaturePluginStateMachine* This,
+		FGameFeaturePluginStateRange InDestination,
+		TDelegate<void(UGameFeaturePluginStateMachine*, const UE::GameFeatures::FResult&)> OnFeatureStateTransitionComplete,
+		FDelegateHandle* OutCallbackHandle = nullptr);
+
+	/** Cancel the current transition if possible */
+	bool TryCancel(UGameFeaturePluginStateMachine* This, TDelegate<void(UGameFeaturePluginStateMachine*)> OnFeatureStateTransitionCanceled, FDelegateHandle* OutCallbackHandle = nullptr) const;
+
+	/** Update the current PluginURL data for this plugin if possible. Returns false if this update fails or
+		if the supplied InPluginURL matches the existing URL data for this plugin.**/
+	UE::GameFeatures::FResult TryUpdatePluginProtocolOptions(UGameFeaturePluginStateMachine* This, const FGameFeatureProtocolOptions& InOptions, bool& bOutDidUpdate) const;
+
+	/** Remove any pending callback from SetDestination */
+	void RemovePendingTransitionCallback(UGameFeaturePluginStateMachine* This, FDelegateHandle InHandle) const;
+
+	/** Remove any pending callback from SetDestination */
+	void RemovePendingTransitionCallback(UGameFeaturePluginStateMachine* This, FDelegateUserObject DelegateObject) const;
+
+	/** Remove any pending callback from TryCancel */
+	void RemovePendingCancelCallback(UGameFeaturePluginStateMachine* This, FDelegateHandle InHandle) const;
+
+	/** Remove any pending callback from TryCancel */
+	void RemovePendingCancelCallback(UGameFeaturePluginStateMachine* This, FDelegateUserObject DelegateObject) const;
+
+	/** Returns the PluginIdentifier used to identify this plugin uniquely to the GameFeaturePlugin Subsystem */
+	const FGameFeaturePluginIdentifier& GetPluginIdentifier(const UGameFeaturePluginStateMachine* This) const;
+
+	/** Returns the name of the game feature. Before StatusKnown, this returns the URL. */
+	const FString& GetGameFeatureName(const UGameFeaturePluginStateMachine* This) const;
+
+	/** Returns the URL */
+	const FString& GetPluginURL(const UGameFeaturePluginStateMachine* This) const;
+
+	/** Returns protocol medata */
+	const FGameFeatureProtocolMetadata& GetProtocolMetadata(const UGameFeaturePluginStateMachine* This) const;
+
+	/** Returns any protocol options */
+	const FGameFeatureProtocolOptions& GetProtocolOptions(const UGameFeaturePluginStateMachine* This) const;
+
+	/** Returns protocol options suitable for reuse by another state machine */
+	FGameFeatureProtocolOptions RecycleProtocolOptions(const UGameFeaturePluginStateMachine* This) const;
+
+	/** Returns the plugin name if known (plugin must have been registered to know the name). */
+	const FString& GetPluginName(const UGameFeaturePluginStateMachine* This) const;
+
+	/** Returns the uplugin filename of the game feature. Before StatusKnown, this returns false. */
+	bool GetPluginFilename(const UGameFeaturePluginStateMachine* This, FString& OutPluginFilename) const;
+
+	/** Returns the enum state for this machine */
+	EGameFeaturePluginState GetCurrentState(const UGameFeaturePluginStateMachine* This) const;
+
+	/** Returns the state range this machine is trying to move to */
+	FGameFeaturePluginStateRange GetDestination(const UGameFeaturePluginStateMachine* This) const;
+
+	/** Returns information about the current state */
+	const FGameFeaturePluginStateInfo& GetCurrentStateInfo(const UGameFeaturePluginStateMachine* This) const;
+
+	/** Returns true if attempting to reach a new destination state */
+	bool IsRunning(const UGameFeaturePluginStateMachine* This) const;
+
+	/** Returns true if the state is at least StatusKnown so we can query info about the game feature plugin */
+	bool IsStatusKnown(const UGameFeaturePluginStateMachine* This) const;
+
+	/** Returns true if the plugin is available to download/load. Only call if IsStatusKnown is true */
+	bool IsAvailable(const UGameFeaturePluginStateMachine* This) const;
+
+	/** Returns true if the plugin is in an error state */
+	bool IsInErrorState(const UGameFeaturePluginStateMachine* This) const;
+
+	/** Whether this machine is allowed to be asynchronous */
+	bool AllowAsyncLoading(const UGameFeaturePluginStateMachine* This) const;
+
+	/** Returns true if the plugin will stream dependencies after installing. Only call if IsStatusKnown is true */
+	bool HasAssetStreamingDependencies(const UGameFeaturePluginStateMachine* This) const;
+
+	void SetWasLoadedAsBuiltIn(UGameFeaturePluginStateMachine* This);
+
+	bool WasLoadedAsBuiltIn(const UGameFeaturePluginStateMachine* This) const;
+
+	/** If the plugin is activated already, we will retrieve its game feature data */
+	UGameFeatureData* GetGameFeatureDataForActivePlugin(UGameFeaturePluginStateMachine* This) const;
+	const UGameFeatureData* GetGameFeatureDataForActivePlugin(const UGameFeaturePluginStateMachine* This) const;
+
+	/** If the plugin is registered already, we will retrieve its game feature data */
+	UGameFeatureData* GetGameFeatureDataForRegisteredPlugin(UGameFeaturePluginStateMachine* This, bool bCheckForRegistering = false) const;
+	const UGameFeatureData* GetGameFeatureDataForRegisteredPlugin(const UGameFeaturePluginStateMachine* This, bool bCheckForRegistering = false) const;
+
+	const FGameFeaturePluginStateMachineProperties& GetProperties(const UGameFeaturePluginStateMachine* This) const;
+
+	/** Returns true if the state machine is in an unrecoverable error state. */
+	bool IsErrorStateUnrecoverable(const UGameFeaturePluginStateMachine* This) const;
+
+	/** Sets that the state machine is in an unrecoverable error. Used when the GFD is missing as that can't be recovered from. */
+	void SetUnrecoverableError(UGameFeaturePluginStateMachine* This) const;
+
+	/** Request to exclude this plugin from batch processing, for e.g. if the plugin is a dependency and has batch processing enabled */
+	void ExcludeFromBatchProcessing(UGameFeaturePluginStateMachine* This);
+
+	/** Helper static function to trigger state specific batch processing functions */
+	static void BatchProcess(EGameFeaturePluginState State, TConstArrayView<UGameFeaturePluginStateMachine*> GFPSMs);	
+
+private:
+	/** Returns true if the specified state is not a transition state */
+	bool IsValidTransitionState(const UGameFeaturePluginStateMachine* This, EGameFeaturePluginState InState) const;
+
+	/** Returns true if the specified state is a destination state */
+	bool IsValidDestinationState(const UGameFeaturePluginStateMachine* This, EGameFeaturePluginState InDestinationState) const;
+
+	/** Returns true if the specified state is a error state */
+	bool IsValidErrorState(const UGameFeaturePluginStateMachine* This, EGameFeaturePluginState InDestinationState) const;
+
+	/** Processes the current state and looks for state transitions */
+	void UpdateStateMachine(UGameFeaturePluginStateMachine* This);
+
+	/** Update Progress for current state */
+	void UpdateCurrentStateProgress(UGameFeaturePluginStateMachine* This, float Progress) const;
+
+	/** Helper to register the state machine with the manager as a state machine in transiton */
+	void RegisterAsTransitioningStateMachine(UGameFeaturePluginStateMachine* This);
+
+	/** Helper to unregister the state machine with the manager as a state machine in transiton */
+	void UnregisterAsTransitioningStateMachine(UGameFeaturePluginStateMachine* This);
+	
+	/** Adds a batching request for current state if state can be batched */
+	void CheckAddBatchingRequestForCurrentState(UGameFeaturePluginStateMachine* This);
+
+	/** Cancels batching re1quest for current state if one exists */
+	void CheckAndCancelBatchingRequestForCurrentState(UGameFeaturePluginStateMachine* This);
+
 };
 
 #undef UE_API
