Use a SharedPtr when caching details to avoid doing deep copies of the structure each time we ask for the details

This was saving ~1-2s on boot for PC

Only changes the internal usage, public usage will still do deep copies to avoid change the public API

#Synced-CL 48801193

diff --git a/Engine/Plugins/Runtime/GameFeatures/Source/GameFeatures/Private/GameFeaturesSubsystem.cpp b/Engine/Plugins/Runtime/GameFeatures/Source/GameFeatures/Private/GameFeaturesSubsystem.cpp
--- a/Engine/Plugins/Runtime/GameFeatures/Source/GameFeatures/Private/GameFeaturesSubsystem.cpp
+++ b/Engine/Plugins/Runtime/GameFeatures/Source/GameFeatures/Private/GameFeaturesSubsystem.cpp
@@ -2470,18 +2470,26 @@ void UGameFeaturesSubsystem::LoadBuiltInGameFeaturePlugin(const TSharedRef<IPlug
 	const FString& PluginDescriptorFilename = Plugin->GetDescriptorFileName();
 
 	FString PluginURL;
-	FGameFeaturePluginDetails PluginDetails;
 
 	if (PluginDescriptorFilename.IsEmpty())
 	{
 		CompleteDelegate.ExecuteIfBound(UE::GameFeatures::FResult(MakeError(UE::GameFeatures::SubsystemErrorNamespace + UE::GameFeatures::CommonErrorCodes::PluginDetailsNotFound)));
+		return;
 	}
 	else if (!GetDefault<UGameFeaturesSubsystemSettings>()->IsValidGameFeaturePlugin(FPaths::ConvertRelativePathToFull(PluginDescriptorFilename)))
 	{
 		// Not a GFP, trivial success
 		CompleteDelegate.ExecuteIfBound(UE::GameFeatures::FResult(MakeValue()));
+		return;
+	}
+	else if (!FPaths::FileExists(PluginDescriptorFilename))
+	{
+		CompleteDelegate.ExecuteIfBound(AddUnmountedPluginExplination(MakeError(UE::GameFeatures::SubsystemErrorNamespace + UE::GameFeatures::CommonErrorCodes::PluginDetailsNotFound)));
+		return;
 	}
-	else if (!FPaths::FileExists(PluginDescriptorFilename) || !GetGameFeaturePluginDetailsInternal(PluginDescriptorFilename, PluginDetails))
+
+	TSharedPtr<FGameFeaturePluginDetails> PluginDetails = GetGameFeaturePluginDetailsInternal(PluginDescriptorFilename);
+	if (PluginDetails == nullptr)
 	{
 		CompleteDelegate.ExecuteIfBound(AddUnmountedPluginExplination(MakeError(UE::GameFeatures::SubsystemErrorNamespace + UE::GameFeatures::CommonErrorCodes::PluginDetailsNotFound)));
 	}
@@ -2503,7 +2511,7 @@ void UGameFeaturesSubsystem::LoadBuiltInGameFeaturePlugin(const TSharedRef<IPlug
 	else
 	{
 		FBuiltInGameFeaturePluginBehaviorOptions BehaviorOptions;
-		const bool bShouldProcess = AdditionalFilter(Plugin->GetDescriptorFileName(), PluginDetails, BehaviorOptions);
+		const bool bShouldProcess = AdditionalFilter(Plugin->GetDescriptorFileName(), *PluginDetails, BehaviorOptions);
 		if (!bShouldProcess)
 		{
 			CompleteDelegate.ExecuteIfBound(AddUnmountedPluginExplination(MakeError(UE::GameFeatures::SubsystemErrorNamespace + UE::GameFeatures::CommonErrorCodes::PluginFiltered)));
@@ -2548,7 +2556,7 @@ void UGameFeaturesSubsystem::LoadBuiltInGameFeaturePlugin(const TSharedRef<IPlug
 			}
 
 			EBuiltInAutoState InitialAutoState = (BehaviorOptions.AutoStateOverride != EBuiltInAutoState::Invalid) ? 
-				BehaviorOptions.AutoStateOverride : PluginDetails.BuiltInAutoState;
+				BehaviorOptions.AutoStateOverride : PluginDetails->BuiltInAutoState;
 			if (InitialAutoState < EBuiltInAutoState::Registered && IsRunningCookCommandlet())
 			{
 				UE_LOG(LogGameFeatures, Display, TEXT("%s will be set to Registered for cooking"), *Plugin->GetName());
@@ -2934,7 +2942,12 @@ bool UGameFeaturesSubsystem::GetBuiltInGameFeaturePluginDetails(const TSharedRef
 	// Make sure you are in a game feature plugins folder. All GameFeaturePlugins are rooted in a GameFeatures folder.
 	if (!PluginDescriptorFilename.IsEmpty() && GetDefault<UGameFeaturesSubsystemSettings>()->IsValidGameFeaturePlugin(FPaths::ConvertRelativePathToFull(PluginDescriptorFilename)))
 	{
-		return GetGameFeaturePluginDetailsInternal(PluginDescriptorFilename, OutPluginDetails);
+		TSharedPtr<FGameFeaturePluginDetails> PluginDetails = GetGameFeaturePluginDetailsInternal(PluginDescriptorFilename);
+		if (PluginDetails)
+		{
+			OutPluginDetails = *PluginDetails;
+			return true;
+		}
 	}
 
 	return false;
@@ -2967,7 +2980,12 @@ bool UGameFeaturesSubsystem::GetGameFeaturePluginDetails(const FString& PluginUR
 	FStringView PluginPath;
 	if (UGameFeaturesSubsystem::ParsePluginURL(PluginURL, nullptr, &PluginPath))
 	{
-		return GetGameFeaturePluginDetailsInternal(FString(PluginPath), OutPluginDetails);
+		TSharedPtr<FGameFeaturePluginDetails> PluginDetails = GetGameFeaturePluginDetailsInternal(FString(PluginPath));
+		if (PluginDetails)
+		{
+			OutPluginDetails = *PluginDetails;
+			return true;
+		}
 	}
 
 	return false;
@@ -2984,14 +3002,14 @@ bool UGameFeaturesSubsystem::GetGameFeatureControlsUPlugin(const FString& Plugin
 	return false;
 }
 
-bool UGameFeaturesSubsystem::GetGameFeaturePluginDetailsInternal(const FString& PluginDescriptorFilename, FGameFeaturePluginDetails& OutPluginDetails) const
+TSharedPtr<FGameFeaturePluginDetails> UGameFeaturesSubsystem::GetGameFeaturePluginDetailsInternal(const FString& PluginDescriptorFilename) const
 {
 	TRACE_CPUPROFILER_EVENT_SCOPE(GFP_GetPluginDetails);
 
 	check(GameSpecificPolicies);
 	if (!GameSpecificPolicies->ShouldReadPluginDetails(PluginDescriptorFilename))
 	{
-		return false;
+		return nullptr;
 	}
 
 	// GFPs are implemented with a plugin so FPluginReferenceDescriptor doesn't know anything about them.
@@ -3001,13 +3019,13 @@ bool UGameFeaturesSubsystem::GetGameFeaturePluginDetailsInternal(const FString&
 	{
 		UE::TReadScopeLock ReadLock(CachedGameFeaturePluginDetailsLock);
 
-		if (const FCachedGameFeaturePluginDetails* ExistingDetails = CachedPluginDetailsByFilename.Find(PluginDescriptorFilename))
+		if (const TSharedPtr<FGameFeaturePluginDetails>* ExistingDetails = CachedPluginDetailsByFilename.Find(PluginDescriptorFilename))
 		{
-			OutPluginDetails = ExistingDetails->Details;
-			return true;
+			return *ExistingDetails;
 		}
 	}
 
+	FGameFeaturePluginDetails PluginDetails;
 
 #if WITH_EDITOR
 	// if we failed to find the plugin and add it to the plugin manager, we may leak this when closing a project
@@ -3054,7 +3072,7 @@ bool UGameFeaturesSubsystem::GetGameFeaturePluginDetailsInternal(const FString&
 				if (!Success && !TryWithNative)
 				{
 					UE_LOG(LogGameFeatures, Error, TEXT("Failed to read plugin descriptor %s. Reason: %s"), *PluginDescriptorFilename, *FailReason.ToString());
-					return false;
+					return nullptr;
 				}
 			}
 
@@ -3065,7 +3083,7 @@ bool UGameFeaturesSubsystem::GetGameFeaturePluginDetailsInternal(const FString&
 				if (!FFileHelper::LoadFileToString(FileContents, *PluginDescriptorFilename))
 				{
 					UE_LOG(LogGameFeatures, Error, TEXT("UGameFeaturesSubsystem could not load plugin descriptor. Failed to read file. File:%s Error:%d"), *PluginDescriptorFilename, FPlatformMisc::GetLastError());
-					return false;
+					return nullptr;
 				}
 
 				// Deserialize a JSON object from the string	
@@ -3073,7 +3091,7 @@ bool UGameFeaturesSubsystem::GetGameFeaturePluginDetailsInternal(const FString&
 				if (!FJsonSerializer::Deserialize(Reader, ObjectPtr) || !ObjectPtr.IsValid())
 				{
 					UE_LOG(LogGameFeatures, Error, TEXT("UGameFeaturesSubsystem could not load plugin descriptor. Json invalid. File:%s. Error:%s"), *PluginDescriptorFilename, *Reader->GetErrorMessage());
-					return false;
+					return nullptr;
 				}
 			}
 		}
@@ -3081,11 +3099,11 @@ bool UGameFeaturesSubsystem::GetGameFeaturePluginDetailsInternal(const FString&
 
 	// Read the properties
 	// Hotfixable. If it is not specified, then we assume it is
-	OutPluginDetails.bHotfixable = true;
-	ObjectPtr->TryGetBoolField(TEXTVIEW("Hotfixable"), OutPluginDetails.bHotfixable);
+	PluginDetails.bHotfixable = true;
+	ObjectPtr->TryGetBoolField(TEXTVIEW("Hotfixable"), PluginDetails.bHotfixable);
 
 	// Determine the initial plugin state
-	OutPluginDetails.BuiltInAutoState = DetermineBuiltInInitialFeatureState(ObjectPtr, PluginDescriptorFilename);
+	PluginDetails.BuiltInAutoState = DetermineBuiltInInitialFeatureState(ObjectPtr, PluginDescriptorFilename);
 
 	// Read any additional metadata the policy might want to consume (e.g., a release version number)
 	for (const FString& ExtraKey : GetDefault<UGameFeaturesSubsystemSettings>()->AdditionalPluginMetadataKeys)
@@ -3093,11 +3111,11 @@ bool UGameFeaturesSubsystem::GetGameFeaturePluginDetailsInternal(const FString&
 		TSharedPtr<FJsonValue> Field = ObjectPtr->TryGetField(ExtraKey);
 		if (Field.IsValid())
 		{
-			OutPluginDetails.AdditionalMetadata.Add(TO_PLUGIN_STRING(ExtraKey), Field);
+			PluginDetails.AdditionalMetadata.Add(TO_PLUGIN_STRING(ExtraKey), Field);
 		}
 		else
 		{
-			OutPluginDetails.AdditionalMetadata.Add(TO_PLUGIN_STRING(ExtraKey), MakeShared<FJsonValueString>(TEXT("")));
+			PluginDetails.AdditionalMetadata.Add(TO_PLUGIN_STRING(ExtraKey), MakeShared<FJsonValueString>(TEXT("")));
 		}
 	}
 
@@ -3144,19 +3162,21 @@ bool UGameFeaturesSubsystem::GetGameFeaturePluginDetailsInternal(const FString&
 			bool bElementActivate = false;
 			ElementObject->TryGetBoolField(ActivateField, bElementActivate);
 
-			FGameFeaturePluginReferenceDetails& RefDetails = OutPluginDetails.PluginDependencies.Emplace_GetRef();
+			FGameFeaturePluginReferenceDetails& RefDetails = PluginDetails.PluginDependencies.Emplace_GetRef();
 			RefDetails.PluginName = TO_PLUGIN_STRING(DependencyName);
 			RefDetails.bShouldActivate = bElementActivate;
 		}
 	}
 
+	TSharedPtr<FGameFeaturePluginDetails> OutPluginDetails = MakeShared<FGameFeaturePluginDetails>(PluginDetails);
+
 	if (UE::GameFeatures::GCachePluginDetails)
 	{
 		UE::TWriteScopeLock WriteLock(CachedGameFeaturePluginDetailsLock);
-		CachedPluginDetailsByFilename.Add(PluginDescriptorFilename, FCachedGameFeaturePluginDetails(OutPluginDetails));
+		CachedPluginDetailsByFilename.Add(PluginDescriptorFilename, OutPluginDetails);
 	}
 
-	return true;
+	return OutPluginDetails;
 }
 
 void UGameFeaturesSubsystem::PruneCachedGameFeaturePluginDetails(const FString& PluginURL, const FString& PluginDescriptorFilename) const
@@ -3974,11 +3994,11 @@ bool UGameFeaturesSubsystem::FindOrCreatePluginDependencyStateMachines(const FSt
 	const bool bWarnOnDepCreation = InStateProperties.ProtocolOptions.bLogWarningOnForcedDependencyCreation;
 	const bool bErrorOnDepCreation = InStateProperties.ProtocolOptions.bLogErrorOnForcedDependencyCreation;
 
-	FGameFeaturePluginDetails Details;
-	if (GetGameFeaturePluginDetailsInternal(PluginFilename, Details))
+	TSharedPtr<FGameFeaturePluginDetails> Details = GetGameFeaturePluginDetailsInternal(PluginFilename);
+	if (Details)
 	{
 		FString DependencyName;
-		for (const FGameFeaturePluginReferenceDetails& PluginDependency : Details.PluginDependencies)
+		for (const FGameFeaturePluginReferenceDetails& PluginDependency : Details->PluginDependencies)
 		{
 			DependencyName = FString(PluginDependency.PluginName);
 			FPluginDependencyDetails DependencyDetails;
@@ -4132,11 +4152,11 @@ bool UGameFeaturesSubsystem::FindPluginDependencyStateMachinesToDeactivate(const
 template <typename CallableT>
 bool UGameFeaturesSubsystem::EnumeratePluginDependenciesWithShouldActivate(const FString& PluginURL, const FString& PluginFilename, CallableT Callable) const
 {
-	FGameFeaturePluginDetails Details;
-	if (GetGameFeaturePluginDetailsInternal(PluginFilename, Details))
+	TSharedPtr<FGameFeaturePluginDetails> Details = GetGameFeaturePluginDetailsInternal(PluginFilename);
+	if (Details)
 	{
 		FString DependencyName;
-		for (const FGameFeaturePluginReferenceDetails& PluginDependency : Details.PluginDependencies)
+		for (const FGameFeaturePluginReferenceDetails& PluginDependency : Details->PluginDependencies)
 		{
 			if (PluginDependency.bShouldActivate)
 			{
diff --git a/Engine/Plugins/Runtime/GameFeatures/Source/GameFeatures/Public/GameFeaturesSubsystem.h b/Engine/Plugins/Runtime/GameFeatures/Source/GameFeatures/Public/GameFeaturesSubsystem.h
--- a/Engine/Plugins/Runtime/GameFeatures/Source/GameFeatures/Public/GameFeaturesSubsystem.h
+++ b/Engine/Plugins/Runtime/GameFeatures/Source/GameFeatures/Public/GameFeaturesSubsystem.h
@@ -852,7 +852,7 @@ private:
 	UE_API const UGameFeatureData* GetRegisteredDataForStateMachine(UGameFeaturePluginStateMachine* GFSM) const;
 
 	/** Gets relevant properties out of a uplugin file */
-	UE_API bool GetGameFeaturePluginDetailsInternal(const FString& PluginDescriptorFilename, struct FGameFeaturePluginDetails& OutPluginDetails) const;
+	UE_API TSharedPtr<struct FGameFeaturePluginDetails> GetGameFeaturePluginDetailsInternal(const FString& PluginDescriptorFilename) const;
 
 	/** Prunes any cached GFP details */
 	UE_API void PruneCachedGameFeaturePluginDetails(const FString& PluginURL, const FString& PluginDescriptorFilename) const;
@@ -984,13 +984,7 @@ private:
 	mutable FTransactionallySafeRWLock GameFeaturePluginNameToPathMapLock;
 	PRAGMA_ENABLE_DEPRECATION_WARNINGS
 
-	struct FCachedGameFeaturePluginDetails
-	{
-		FGameFeaturePluginDetails Details;
-		FCachedGameFeaturePluginDetails() = default;
-		FCachedGameFeaturePluginDetails(const FGameFeaturePluginDetails& InDetails) : Details(InDetails) {}
-	};
-	mutable TMap<FString, FCachedGameFeaturePluginDetails> CachedPluginDetailsByFilename;
+	mutable TMap<FString, TSharedPtr<FGameFeaturePluginDetails>> CachedPluginDetailsByFilename;
 	mutable FTransactionallySafeRWLock CachedGameFeaturePluginDetailsLock;
 
 	UPROPERTY()
