Logging / GameFeatures:  Fixed a couple minor typos.

#rnx

#Synced-CL 48947878

diff --git a/Engine/Plugins/Runtime/GameFeatures/Source/GameFeatures/Private/GameFeatureData.cpp b/Engine/Plugins/Runtime/GameFeatures/Source/GameFeatures/Private/GameFeatureData.cpp
--- a/Engine/Plugins/Runtime/GameFeatures/Source/GameFeatures/Private/GameFeatureData.cpp
+++ b/Engine/Plugins/Runtime/GameFeatures/Source/GameFeatures/Private/GameFeatureData.cpp
@@ -297,7 +297,7 @@ void UGameFeatureData::InitializeHierarchicalPluginIniFiles(const FString& Plugi
 				for (const TPair<FName, TSet<FString>>& Pair : ChangeTracker->ModifiedSectionsPerBranch)
 				{
 					ensureAlwaysMsgf(!Pair.Value.Contains(TEXT("ConsoleVariables")) && !Pair.Value.Contains(TEXT("ConsoleVariables_HighPriority")),
-						TEXT("GameFeaturePlugins .ini files should no longer use [ConsoleVariables] or [ConsoleVariables_HighPriority] sections. Instead, use [PluginCVars_LowPriority] and [PluginCVars_HighPriority]. The old naming is in one of more of: %s"),
+						TEXT("GameFeaturePlugins .ini files should no longer use [ConsoleVariables] or [ConsoleVariables_HighPriority] sections. Instead, use [PluginCVars_LowPriority] and [PluginCVars_HighPriority]. The old naming is in one or more of: '%s'."),
 						*FString::Join(ChangeTracker->LoadedFiles, TEXT(", ")));
 				}
 #endif
@@ -728,7 +728,7 @@ void UGameFeatureData::InitializeHierarchicalPluginIniFiles(const TArrayView<FSt
 				for (const TPair<FName, TSet<FString>>& Pair : ChangeTracker->ModifiedSectionsPerBranch)
 				{
 					ensureAlwaysMsgf(!Pair.Value.Contains(TEXT("ConsoleVariables")) && !Pair.Value.Contains(TEXT("ConsoleVariables_HighPriority")),
-						TEXT("GameFeaturePlugins .ini files should no longer use [ConsoleVariables] or [ConsoleVariables_HighPriority] sections. Instead, use [PluginCVars_LowPriority] and [PluginCVars_HighPriority]. The old naming is in one of more of: %s"),
+						TEXT("GameFeaturePlugins .ini files should no longer use [ConsoleVariables] or [ConsoleVariables_HighPriority] sections. Instead, use [PluginCVars_LowPriority] and [PluginCVars_HighPriority]. The old naming is in one or more of: '%s'."),
 						*FString::Join(ChangeTracker->LoadedFiles, TEXT(", ")));
 				}
 #endif
