Fix issue where cancelling a GFP operation can leave plugins in a state where they can't be installed again.
Adding function to check if a plugin is in an error state to allow for gameplay abilities to know when a plugin is in error state but also installed.

diff --git a/Engine/Plugins/Runtime/GameFeatures/Source/GameFeatures/Private/GameFeaturesSubsystem.cpp b/Engine/Plugins/Runtime/GameFeatures/Source/GameFeatures/Private/GameFeaturesSubsystem.cpp
--- a/Engine/Plugins/Runtime/GameFeatures/Source/GameFeatures/Private/GameFeaturesSubsystem.cpp
+++ b/Engine/Plugins/Runtime/GameFeatures/Source/GameFeatures/Private/GameFeaturesSubsystem.cpp
@@ -1440,6 +1440,15 @@ const UGameFeatureData* UGameFeaturesSubsystem::GetGameFeatureDataForRegisteredP
 	return nullptr;
 }
 
+bool UGameFeaturesSubsystem::IsGameFeaturePluginInErrorState(const FString& PluginURL) const
+{
+	if (const UGameFeaturePluginStateMachine* StateMachine = FindGameFeaturePluginStateMachine(PluginURL))
+	{
+		return StateMachine->IsInErrorState();
+	}
+	return false;
+}
+
 bool UGameFeaturesSubsystem::IsGameFeaturePluginInstalled(const FString& PluginURL) const
 {
 	if (const UGameFeaturePluginStateMachine* StateMachine = FindGameFeaturePluginStateMachine(PluginURL))
@@ -2383,6 +2392,32 @@ void UGameFeaturesSubsystem::TerminateGameFeaturePlugin(TConstArrayView<FString>
 	}
 }
 
+bool UGameFeaturesSubsystem::CancellAmortizedLoads(const FString& PluginURL)
+{
+	bool bDidCancellAnything = false;
+	const FString CancellingPluginName = GetPluginNameFromPluginURL(PluginURL);
+	if (!CancellingPluginName.IsEmpty())
+	{
+		for (TPair<FGuid, TWeakPtr<FAmmortizedLoadContext>>& AmmortizedLoad : ActiveAmmoritzedLoads)
+		{
+			if (AmmortizedLoad.Value.IsValid())
+			{
+				for (const TPair<FString, UE::GameFeatures::FResult>& PluginToResult : AmmortizedLoad.Value.Pin()->Results)
+				{
+					if (PluginToResult.Key.Equals(CancellingPluginName, ESearchCase::IgnoreCase))
+					{
+						AmmortizedLoad.Value.Pin()->bCancelled = true;
+						bDidCancellAnything = true;
+						break;
+					}
+				}
+			}
+		}
+	}
+
+	return bDidCancellAnything;
+}
+
 void UGameFeaturesSubsystem::CancelGameFeatureStateChange(const FString& PluginURL)
 {
 	CancelGameFeatureStateChange(PluginURL, FGameFeaturePluginChangeStateComplete());
@@ -2390,6 +2425,8 @@ void UGameFeaturesSubsystem::CancelGameFeatureStateChange(const FString& PluginU
 
 void UGameFeaturesSubsystem::CancelGameFeatureStateChange(const FString& PluginURL, const FGameFeaturePluginChangeStateComplete& CompleteDelegate)
 {
+	CancellAmortizedLoads(PluginURL);
+
 	if (UGameFeaturePluginStateMachine* StateMachine = FindGameFeaturePluginStateMachine(PluginURL))
 	{
 		const bool bCancelPending = StateMachine->TryCancel(FGameFeatureStateTransitionCanceled::CreateWeakLambda(this, [CompleteDelegate](UGameFeaturePluginStateMachine* Machine)
@@ -2692,54 +2729,11 @@ void UGameFeaturesSubsystem::LoadBuiltInGameFeaturePluginsInternal(TArray<TShare
 {
 	TRACE_CPUPROFILER_EVENT_SCOPE_STR("UGameFeaturesSubsystem::LoadBuiltInGameFeaturePlugins");
 
-	struct FLoadContext
-	{
-		FScopeLogTime ScopeLogTime{TEXT("BuiltInGameFeaturePlugins loaded."), nullptr, FConditionalScopeLogTime::ScopeLog_Seconds};
-
-		TArray<TSharedRef<IPlugin>> AmortizedPluginsRemaining;
-		TMap<FString, UE::GameFeatures::FResult> Results;
-		FBuiltInGameFeaturePluginsLoaded CompleteDelegate;
+	TSharedRef<FAmmortizedLoadContext> LoadContext = MakeShared<FAmmortizedLoadContext>();
+	FGuid ReferenceGuid = FGuid::NewGuid();
+	LoadContext->GUID = ReferenceGuid;
+	ActiveAmmoritzedLoads.Add(ReferenceGuid, LoadContext);
 
-		int32 NumPluginsLoaded = 0;
-		int32 PluginAmortizeRate = 0;
-		bool bPushedTagsBroadcast = false;
-		bool bPushedAssetBulkScanning = false;
-
-		FLoadContext()
-		{
-			if (UAssetManager* AssetManager = UAssetManager::GetIfInitialized())
-			{
-				AssetManager->PushBulkScanning();
-				bPushedAssetBulkScanning = true;
-			}
-			if (UGameplayTagsManager* TagsManager = UGameplayTagsManager::GetIfAllocated())
-			{
-				TagsManager->PushDeferOnGameplayTagTreeChangedBroadcast();
-				bPushedTagsBroadcast = true;
-			}
-		}
-
-		~FLoadContext()
-		{
-			if (bPushedTagsBroadcast)
-			{
-				if (UGameplayTagsManager* TagsManager = UGameplayTagsManager::GetIfAllocated())
-				{
-					TagsManager->PopDeferOnGameplayTagTreeChangedBroadcast();
-				}
-			}
-			if (bPushedAssetBulkScanning)
-			{
-				if (UAssetManager* AssetManager = UAssetManager::GetIfInitialized())
-				{
-					AssetManager->PopBulkScanning();
-				}
-			}
-
-			CompleteDelegate.ExecuteIfBound(Results);
-		}
-	};
-	TSharedRef<FLoadContext> LoadContext = MakeShared<FLoadContext>();
 	LoadContext->CompleteDelegate = InCompleteDelegate;
 
 	FBuiltInPluginLoadTimeTracker PluginLoadTimeTracker;
@@ -2777,7 +2771,7 @@ void UGameFeaturesSubsystem::LoadBuiltInGameFeaturePluginsInternal(TArray<TShare
 
 	Algo::TopologicalSort(Plugins, GetPluginDependencies);
 
-	auto ProcessGameFeaturePlugin = [this](const TSharedRef<IPlugin>& PluginToProcess, const TSharedRef<FLoadContext>& LoadContextToProcess, const FBuiltInPluginAdditionalFilters& AdditionalFilterForProcess)
+	auto ProcessGameFeaturePlugin = [this](const TSharedRef<IPlugin>& PluginToProcess, const TSharedRef<FAmmortizedLoadContext>& LoadContextToProcess, const FBuiltInPluginAdditionalFilters& AdditionalFilterForProcess)
 	{
 		LoadBuiltInGameFeaturePlugin(PluginToProcess, AdditionalFilterForProcess, FGameFeaturePluginLoadComplete::CreateLambda([LoadContextToProcess, PluginToProcess](const UE::GameFeatures::FResult& Result)
 		{
@@ -2794,7 +2788,7 @@ void UGameFeaturesSubsystem::LoadBuiltInGameFeaturePluginsInternal(TArray<TShare
 		// Note that we are using AdditionalFilter_Copyable here. We need to make a copy of thus function so that the captured properties in lambdas are not destroyed
 		FTSTicker::GetCoreTicker().AddTicker(FTickerDelegate::CreateWeakLambda(this, [this, LoadContext, AdditionalFilter_Copyable, ProcessGameFeaturePlugin](float)
 		{
-			if (GameSpecificPolicies)
+			if (GameSpecificPolicies && !LoadContext->bCancelled)
 			{
 				int32 NumPluginsThisFrame = FMath::Min(LoadContext->PluginAmortizeRate, LoadContext->AmortizedPluginsRemaining.Num());
 				
@@ -2802,12 +2796,11 @@ void UGameFeaturesSubsystem::LoadBuiltInGameFeaturePluginsInternal(TArray<TShare
 				{
 					const TSharedRef<IPlugin>& Plugin = LoadContext->AmortizedPluginsRemaining[PluginIdx];
 					ProcessGameFeaturePlugin(Plugin, LoadContext, AdditionalFilter_Copyable);
-					
 				}
 
 				LoadContext->AmortizedPluginsRemaining.RemoveAt(0, NumPluginsThisFrame);
 			}
-			return GameSpecificPolicies && LoadContext->AmortizedPluginsRemaining.Num() > 0;
+			return GameSpecificPolicies && LoadContext->AmortizedPluginsRemaining.Num() > 0 && !LoadContext->bCancelled;
 		}));
 	}
 	else
@@ -4808,3 +4801,38 @@ void UGameFeaturesSubsystem::SetExplanationForNotMountingPlugin(const FString& P
 	}
 #endif
 }
+
+UGameFeaturesSubsystem::FAmmortizedLoadContext::FAmmortizedLoadContext()
+{
+	if (UAssetManager* AssetManager = UAssetManager::GetIfInitialized())
+	{
+		AssetManager->PushBulkScanning();
+		bPushedAssetBulkScanning = true;
+	}
+	if (UGameplayTagsManager* TagsManager = UGameplayTagsManager::GetIfAllocated())
+	{
+		TagsManager->PushDeferOnGameplayTagTreeChangedBroadcast();
+		bPushedTagsBroadcast = true;
+	}
+}
+
+UGameFeaturesSubsystem::FAmmortizedLoadContext::~FAmmortizedLoadContext()
+{
+	if (bPushedTagsBroadcast)
+	{
+		if (UGameplayTagsManager* TagsManager = UGameplayTagsManager::GetIfAllocated())
+		{
+			TagsManager->PopDeferOnGameplayTagTreeChangedBroadcast();
+		}
+	}
+	if (bPushedAssetBulkScanning)
+	{
+		if (UAssetManager* AssetManager = UAssetManager::GetIfInitialized())
+		{
+			AssetManager->PopBulkScanning();
+		}
+	}
+
+	CompleteDelegate.ExecuteIfBound(Results);
+	UGameFeaturesSubsystem::Get().ActiveAmmoritzedLoads.Remove(GUID);
+}
diff --git a/Engine/Plugins/Runtime/GameFeatures/Source/GameFeatures/Public/GameFeaturesSubsystem.h b/Engine/Plugins/Runtime/GameFeatures/Source/GameFeatures/Public/GameFeaturesSubsystem.h
--- a/Engine/Plugins/Runtime/GameFeatures/Source/GameFeatures/Public/GameFeaturesSubsystem.h
+++ b/Engine/Plugins/Runtime/GameFeatures/Source/GameFeatures/Public/GameFeaturesSubsystem.h
@@ -6,7 +6,9 @@
 #include "Delegates/Delegate.h"
 #include "Engine/Engine.h"
 #include "GameFeatureTypesFwd.h"
+#include "GameFeaturePluginOperationResult.h"
 #include "Misc/TransactionallySafeRWLock.h"
+#include "Stats/StatsMisc.h"
 
 #include "GameFeaturesSubsystem.generated.h"
 
@@ -580,6 +582,9 @@ public:
 	/** Returns the game feature data for a registered plugin specified by PluginURL */
 	UE_API const UGameFeatureData* GetGameFeatureDataForRegisteredPluginByURL(const FString& PluginURL, bool bCheckForRegistering = false);
 
+	/** Determines if a plugin is in an error state */
+	UE_API bool IsGameFeaturePluginInErrorState(const FString& PluginURL) const;
+
 	/** Determines if a plugin is in the Installed state (or beyond) */
 	UE_API bool IsGameFeaturePluginInstalled(const FString& PluginURL) const;
 
@@ -1010,6 +1015,29 @@ private:
 	bool bInitializedPolicyManager = false;
 
 	static UE_API TSet<FString> DlcPlugins;
+
+	struct FAmmortizedLoadContext
+	{
+		FScopeLogTime ScopeLogTime{ TEXT("BuiltInGameFeaturePlugins loaded."), nullptr, FConditionalScopeLogTime::ScopeLog_Seconds };
+
+		TArray<TSharedRef<IPlugin>> AmortizedPluginsRemaining;
+		TMap<FString, UE::GameFeatures::FResult> Results;
+		FBuiltInGameFeaturePluginsLoaded CompleteDelegate;
+
+		int32 NumPluginsLoaded = 0;
+		int32 PluginAmortizeRate = 0;
+		bool bPushedTagsBroadcast = false;
+		bool bPushedAssetBulkScanning = false;
+		bool bCancelled = false;
+		FGuid GUID;
+
+		FAmmortizedLoadContext();
+
+		~FAmmortizedLoadContext();
+	};
+
+	bool CancellAmortizedLoads(const FString& PluginURL);
+	TMap<FGuid, TWeakPtr<FAmmortizedLoadContext>> ActiveAmmoritzedLoads;
 };
 
 #undef UE_API
