When doing amortized GFP loading, use a set ms time we should spend per tick on loading GFPs
This is not perfect and does not cover defered work caused by GFPs but based on insights the general extra
Use for now how much time it takes to process each GFP as "work"
This generally looks like an upper bounds of ~2x extra work can be added. So if we aim at 100ms time rate, we can expect some 200ms frames

#Synced-CL 49025350

diff --git a/Engine/Plugins/Runtime/GameFeatures/Source/GameFeatures/Private/GameFeaturesSubsystem.cpp b/Engine/Plugins/Runtime/GameFeatures/Source/GameFeatures/Private/GameFeaturesSubsystem.cpp
--- a/Engine/Plugins/Runtime/GameFeatures/Source/GameFeatures/Private/GameFeaturesSubsystem.cpp
+++ b/Engine/Plugins/Runtime/GameFeatures/Source/GameFeatures/Private/GameFeaturesSubsystem.cpp
@@ -2781,7 +2781,7 @@ void UGameFeaturesSubsystem::LoadBuiltInGameFeaturePluginsInternal(TArray<TShare
 		}));
 	};
 
-	if (AmortizeRate > 0)
+	if (AmortizeRate > 0 && NumPluginsToLoad > 0)
 	{
 		LoadContext->PluginAmortizeRate = AmortizeRate;
 		LoadContext->AmortizedPluginsRemaining = Plugins;
@@ -2790,13 +2790,17 @@ void UGameFeaturesSubsystem::LoadBuiltInGameFeaturePluginsInternal(TArray<TShare
 		{
 			if (GameSpecificPolicies && !LoadContext->bCancelled)
 			{
-				int32 NumPluginsThisFrame = FMath::Min(LoadContext->PluginAmortizeRate, LoadContext->AmortizedPluginsRemaining.Num());
-				
-				for (int32 PluginIdx = 0; PluginIdx < NumPluginsThisFrame; ++PluginIdx)
+				// This PluginAmortizeRate use to be how many GFPs per tick, but we have moved this rate to mean milliseconds now
+				double MaxMillisecondsPerFrame = LoadContext->PluginAmortizeRate / 1000.0;
+				double StartTime = FPlatformTime::Seconds();
+				int32 NumPluginsThisFrame = 0;
+				int32 NumPluginsLeft = LoadContext->AmortizedPluginsRemaining.Num();
+
+				do
 				{
-					const TSharedRef<IPlugin>& Plugin = LoadContext->AmortizedPluginsRemaining[PluginIdx];
+					const TSharedRef<IPlugin>& Plugin = LoadContext->AmortizedPluginsRemaining[NumPluginsThisFrame++];
 					ProcessGameFeaturePlugin(Plugin, LoadContext, AdditionalFilter_Copyable);
-				}
+				} while (FPlatformTime::Seconds() - StartTime < MaxMillisecondsPerFrame && NumPluginsThisFrame < NumPluginsLeft);
 
 				LoadContext->AmortizedPluginsRemaining.RemoveAt(0, NumPluginsThisFrame);
 			}
