When bHasAssetDependencies is false call the OnProgress callback if available.

diff --git a/Engine/Plugins/Runtime/GameFeatures/Source/GameFeatures/Private/GameFeaturesSubsystem.cpp b/Engine/Plugins/Runtime/GameFeatures/Source/GameFeatures/Private/GameFeaturesSubsystem.cpp
--- a/Engine/Plugins/Runtime/GameFeatures/Source/GameFeatures/Private/GameFeaturesSubsystem.cpp
+++ b/Engine/Plugins/Runtime/GameFeatures/Source/GameFeatures/Private/GameFeaturesSubsystem.cpp
@@ -3705,6 +3705,10 @@ struct FGameFeaturePluginPredownloadContext : public FGameFeaturePluginPredownlo
 		if (!bHasAssetDependencies)
 		{
 			Progress = InstallBundleProgress;
+			if (OnProgress)
+			{
+				OnProgress(Progress);
+			}
 			return;
 		}
 
