Add function to return the minimal required state for a Queried plugin
Only deactivate a depends if we are not ref counted by some other GFP

This will fix issues where you have:
A (active) -> B (active) <- C (active)
A (registered) -> B (loaded) <- C (active)

diff --git a/Engine/Plugins/Runtime/GameFeatures/Source/GameFeatures/Private/GameFeatureStateHandleReferenceController.cpp b/Engine/Plugins/Runtime/GameFeatures/Source/GameFeatures/Private/GameFeatureStateHandleReferenceController.cpp
--- a/Engine/Plugins/Runtime/GameFeatures/Source/GameFeatures/Private/GameFeatureStateHandleReferenceController.cpp
+++ b/Engine/Plugins/Runtime/GameFeatures/Source/GameFeatures/Private/GameFeatureStateHandleReferenceController.cpp
@@ -318,6 +318,18 @@ void FGameFeatureStateHandleReferenceController::AddOrUpdateReference(const FGam
 	}
 }
 
+/** From the RefCounts return the lowest GFP state required by a GFP we are tracking */
+bool FGameFeatureStateHandleReferenceController::GetMinimumRequiredStateForGameFeaturePlugin(const FString& GameFeaturePluginURL, EGameFeaturePluginState& OutPluginState)
+{
+	const FGameFeatureStateRefCount* RefCount = Internal->GFPReferenceCount.Find(GameFeaturePluginURL);
+	if (RefCount)
+	{
+		return RefCount->GetHighestRefCountDestPluginState(OutPluginState);
+	}
+
+	return false;
+}
+
 bool FGameFeatureStateHandleReferenceController::GetDowngradedTargetState(const FString& PluginURL, EGameFeaturePluginState& OutTransitionState) const
 {
 	const FGameFeatureStateRefCount* RefCount = Internal->GFPReferenceCount.Find(PluginURL);
diff --git a/Engine/Plugins/Runtime/GameFeatures/Source/GameFeatures/Private/GameFeaturesSubsystem.cpp b/Engine/Plugins/Runtime/GameFeatures/Source/GameFeatures/Private/GameFeaturesSubsystem.cpp
--- a/Engine/Plugins/Runtime/GameFeatures/Source/GameFeatures/Private/GameFeaturesSubsystem.cpp
+++ b/Engine/Plugins/Runtime/GameFeatures/Source/GameFeatures/Private/GameFeaturesSubsystem.cpp
@@ -16,6 +16,7 @@
 #include "GameFeatureData.h"
 #include "GameFeaturePluginStateMachine.h"
 #include "GameFeatureStateChangeObserver.h"
+#include "GameFeatureStateHandleReferenceController.h"
 #include "GameplayTagsManager.h"
 #include "HAL/Platform.h"
 #include "Interfaces/IPluginManager.h"
@@ -4135,7 +4136,19 @@ bool UGameFeaturesSubsystem::FindPluginDependencyStateMachinesToDeactivate(const
 		UGameFeaturePluginStateMachine* Dependency = FindGameFeaturePluginStateMachine(DependencyURL);
 		if (Dependency)
 		{
-			OutDependencyMachines.Add(Dependency);
+			// deactivate by default will mean loaded, so make that our default state as we check our min required from the StateHandle, if there is a ref count
+			EGameFeaturePluginState MinRequiredState = EGameFeaturePluginState::Loaded;
+			FGameFeatureStateHandleReferenceController::Get().GetMinimumRequiredStateForGameFeaturePlugin(DependencyURL, MinRequiredState);
+
+			// only if our min required state is less then active should allow these to deactivate, otherwise something still has a ref count on this depends being active
+			if (MinRequiredState <= EGameFeaturePluginState::Loaded)
+			{
+				OutDependencyMachines.Add(Dependency);
+			}
+			else
+			{
+				UE_LOG(LogGameFeatures, Verbose, TEXT("FindPluginDependencyStateMachinesToDeactivate skipped '%s' as it still has an 'Active' minimum reuqired state"), *DependencyName);
+			}
 		}
 		else
 		{
diff --git a/Engine/Plugins/Runtime/GameFeatures/Source/GameFeatures/Public/GameFeatureStateHandleReferenceController.h b/Engine/Plugins/Runtime/GameFeatures/Source/GameFeatures/Public/GameFeatureStateHandleReferenceController.h
--- a/Engine/Plugins/Runtime/GameFeatures/Source/GameFeatures/Public/GameFeatureStateHandleReferenceController.h
+++ b/Engine/Plugins/Runtime/GameFeatures/Source/GameFeatures/Public/GameFeatureStateHandleReferenceController.h
@@ -44,6 +44,9 @@ public:
 	/** Adds or Updates the GFP URL to the PluginState */
 	UE_API void AddOrUpdateReference(const FGameFeatureStateHandle& StateHandle, const FString& GameFeaturePluginURL, EGameFeaturePluginState PluginState, bool bCollectDepends = true);
 
+	/** From the RefCounts gets the lowest GFP state required by a GFP we are tracking into OutPluginState */
+	UE_API bool GetMinimumRequiredStateForGameFeaturePlugin(const FString& GameFeaturePluginURL, EGameFeaturePluginState& OutPluginState);
+
 private:
 	FGameFeatureStateHandleReferenceController();
 	FGameFeatureStateHandleReferenceController(FGameFeatureStateHandleReferenceController const&) = delete;
