Add null check to pervent GEngine access during shutdown.

diff --git a/Engine/Plugins/Runtime/GameFeatures/Source/GameFeatures/Private/GameFeaturesSubsystem.cpp b/Engine/Plugins/Runtime/GameFeatures/Source/GameFeatures/Private/GameFeaturesSubsystem.cpp
--- a/Engine/Plugins/Runtime/GameFeatures/Source/GameFeatures/Private/GameFeaturesSubsystem.cpp
+++ b/Engine/Plugins/Runtime/GameFeatures/Source/GameFeatures/Private/GameFeaturesSubsystem.cpp
@@ -4855,5 +4855,8 @@ UGameFeaturesSubsystem::FAmmortizedLoadContext::~FAmmortizedLoadContext()
 	}
 
 	CompleteDelegate.ExecuteIfBound(Results);
-	UGameFeaturesSubsystem::Get().ActiveAmmoritzedLoads.Remove(GUID);
+	if (GEngine)
+	{
+		UGameFeaturesSubsystem::Get().ActiveAmmoritzedLoads.Remove(GUID);
+	}
 }
