On exit we may have some FAmmortizedLoadContext still around while the StateMachine is working, avoid trying to get the GameFeatureSubsystem Singleton if we are dieing

diff --git a/Engine/Plugins/Runtime/GameFeatures/Source/GameFeatures/Private/GameFeaturesSubsystem.cpp b/Engine/Plugins/Runtime/GameFeatures/Source/GameFeatures/Private/GameFeaturesSubsystem.cpp
--- a/Engine/Plugins/Runtime/GameFeatures/Source/GameFeatures/Private/GameFeaturesSubsystem.cpp
+++ b/Engine/Plugins/Runtime/GameFeatures/Source/GameFeatures/Private/GameFeaturesSubsystem.cpp
@@ -4855,7 +4855,7 @@ UGameFeaturesSubsystem::FAmmortizedLoadContext::~FAmmortizedLoadContext()
 	}
 
 	CompleteDelegate.ExecuteIfBound(Results);
-	if (GEngine)
+	if (!IsEngineExitRequested())
 	{
 		UGameFeaturesSubsystem::Get().ActiveAmmoritzedLoads.Remove(GUID);
 	}
