Enable StateHandles by default everywhere (Server, Client Editor). SEW has it disabled for now due to an issue where the ownership is not properly being released

diff --git a/Engine/Plugins/Runtime/GameFeatures/Source/GameFeatures/Private/GameFeatureStateHandleReferenceController.cpp b/Engine/Plugins/Runtime/GameFeatures/Source/GameFeatures/Private/GameFeatureStateHandleReferenceController.cpp
--- a/Engine/Plugins/Runtime/GameFeatures/Source/GameFeatures/Private/GameFeatureStateHandleReferenceController.cpp
+++ b/Engine/Plugins/Runtime/GameFeatures/Source/GameFeatures/Private/GameFeatureStateHandleReferenceController.cpp
@@ -17,7 +17,7 @@ DEFINE_LOG_CATEGORY_STATIC(LogGameFeatureStateHandle, Log, All);
 
 // A quick way for now to disable and fallback to how things are now
 // something to remove as this system is proved out
-static bool GGameFeatureStateHandleEnabled = UE_CLIENT || UE_SERVER;
+static bool GGameFeatureStateHandleEnabled = true;
 static FAutoConsoleVariableRef CVarGameFeatureStateHandleEnabled(
 	TEXT("GameFeatureStateHandle.Enabled"),
 	GGameFeatureStateHandleEnabled,
