Explicitly delete FGameFeaturePluginState objects in UGameFeaturesSubsystem::Deinitialize to prevent them from calling UGameFeaturesSubsystem::Get() during their destructors

#Synced-CL 49669520

diff --git a/Engine/Plugins/Runtime/GameFeatures/Source/GameFeatures/Private/GameFeaturesSubsystem.cpp b/Engine/Plugins/Runtime/GameFeatures/Source/GameFeatures/Private/GameFeaturesSubsystem.cpp
--- a/Engine/Plugins/Runtime/GameFeatures/Source/GameFeatures/Private/GameFeaturesSubsystem.cpp
+++ b/Engine/Plugins/Runtime/GameFeatures/Source/GameFeatures/Private/GameFeaturesSubsystem.cpp
@@ -696,6 +696,30 @@ void UGameFeaturesSubsystem::Deinitialize()
 #if WITH_EDITOR
 	GameFeatureDataExternalAssetsPathCache.Reset();
 #endif
+
+	// Delete all state objects so that they don't try to run destructor logic during GC/shutdown in a way that could access 
+	// already-deleted objects (e.g. GEngine, this subsystem)
+	for (TPair<FString, TObjectPtr<UGameFeaturePluginStateMachine>>& Pair : GameFeaturePluginStateMachines)
+	{
+		if (Pair.Value)
+		{
+			for (TUniquePtr<FGameFeaturePluginState>& Ptr : Pair.Value->AllStates)
+			{
+				Ptr.Reset();
+			}
+		}
+	}
+
+	for (TObjectPtr<UGameFeaturePluginStateMachine> Terminal : TerminalGameFeaturePluginStateMachines)
+	{
+		if (Terminal)
+		{
+			for (TUniquePtr<FGameFeaturePluginState>& Ptr : Terminal->AllStates)
+			{
+				Ptr.Reset();
+			}
+		}
+	}
 }
 
 void UGameFeaturesSubsystem::OnAssetManagerCreated()
