[GameFeatures] Rework usage of state handles to prevent accidentally using reallocated memory
* Ensure that on complete is always called when remove references isn't.

#tests preflights

#Synced-CL 49755375

diff --git a/Engine/Plugins/Runtime/GameFeatures/Source/GameFeatures/Private/GameFeatureStateHandleReferenceController.cpp b/Engine/Plugins/Runtime/GameFeatures/Source/GameFeatures/Private/GameFeatureStateHandleReferenceController.cpp
--- a/Engine/Plugins/Runtime/GameFeatures/Source/GameFeatures/Private/GameFeatureStateHandleReferenceController.cpp
+++ b/Engine/Plugins/Runtime/GameFeatures/Source/GameFeatures/Private/GameFeatureStateHandleReferenceController.cpp
@@ -633,23 +633,25 @@ void FGameFeatureStateHandleReferenceController::ResetGameFeatureStateHandle(con
 	if (!IsLifetimeControlEnabled())
 	{
 		OnComplete(true);
-		return;
 	}
-
+	else if (IsEngineExitRequested())
+	{
+		UE_LOG(LogGameFeatureStateHandle, Verbose, TEXT("Requesting a reset but the Engine is Exiting. This should have likely been done before we existed the Engine. StateHandle Owner '%s'"), *StateHandle.ToString());
+		OnComplete(true);
+	}
+	else
 	{
 		FScopeLock Lock(&Internal->StateHandlesLock);
 		if (FGameFeatureStateHandleInternal* InternalHandle = Internal->StateHandles.Find(StateHandle))
 		{
-			if (IsEngineExitRequested())
-			{
-				UE_LOG(LogGameFeatureStateHandle, Verbose, TEXT("Requesting a reset but the Engine is Exiting. This should have likely been done before we existed the Engine. StateHandle Owner '%s'"), *StateHandle.ToString());
-			}
-			else
-			{
-				RemoveReferences(StateHandle, InternalHandle->GetPlugins(), OnComplete);
-			}
-
+			const TArray<FString> Plugins = InternalHandle->GetPlugins();
 			InternalHandle->Empty();
+
+			RemoveReferences(StateHandle, Plugins, OnComplete);
+		}
+		else
+		{
+			OnComplete(false);
 		}
 	}
 }
@@ -659,23 +661,24 @@ void FGameFeatureStateHandleReferenceController::ResetSingleGameFeatureStateHand
 	if (!IsLifetimeControlEnabled())
 	{
 		OnComplete(true);
-		return;
 	}
-
+	else if (IsEngineExitRequested())
+	{
+		UE_LOG(LogGameFeatureStateHandle, Verbose, TEXT("Requesting a reset on %s but the Engine is Exiting. This should have likely been done before we existed the Engine. StateHandle Owner '%s'"), *PluginURL, *StateHandle.ToString());
+		OnComplete(true);
+	}
+	else
 	{
 		FScopeLock Lock(&Internal->StateHandlesLock);
 		if (FGameFeatureStateHandleInternal* InternalHandle = Internal->StateHandles.Find(StateHandle))
 		{
-			if (IsEngineExitRequested())
-			{
-				UE_LOG(LogGameFeatureStateHandle, Verbose, TEXT("Requesting a reset on %s but the Engine is Exiting. This should have likely been done before we existed the Engine. StateHandle Owner '%s'"), *PluginURL, *StateHandle.ToString());
-			}
-			else
-			{
-				RemoveReferences(StateHandle, TArray<FString>{ PluginURL }, OnComplete);
-			}
-
 			InternalHandle->Remove(PluginURL);
+
+			RemoveReferences(StateHandle, TArray<FString>{ PluginURL }, OnComplete);
+		}
+		else
+		{
+			OnComplete(false);
 		}
 	}
 }
