[GFP] Add state name to State transition CPU scopes

#rnx

#Synced-CL 49835867

diff --git a/Engine/Plugins/Runtime/GameFeatures/Source/GameFeatures/Private/GameFeaturePluginStateMachine.cpp b/Engine/Plugins/Runtime/GameFeatures/Source/GameFeatures/Private/GameFeaturePluginStateMachine.cpp
--- a/Engine/Plugins/Runtime/GameFeatures/Source/GameFeatures/Private/GameFeaturePluginStateMachine.cpp
+++ b/Engine/Plugins/Runtime/GameFeatures/Source/GameFeatures/Private/GameFeaturePluginStateMachine.cpp
@@ -4539,6 +4539,18 @@ UE_TRACE_EVENT_BEGIN(Cpu, GFP_UpdateStateMachine, NoSync)
 	UE_TRACE_EVENT_FIELD(UE::Trace::WideString, PluginName)
 UE_TRACE_EVENT_END()
 
+UE_TRACE_EVENT_BEGIN(Cpu, GFP_UpdateStateMachine_BeginState, NoSync)
+	UE_TRACE_EVENT_FIELD(UE::Trace::WideString, State)
+UE_TRACE_EVENT_END()
+
+UE_TRACE_EVENT_BEGIN(Cpu, GFP_UpdateStateMachine_UpdateState, NoSync)
+	UE_TRACE_EVENT_FIELD(UE::Trace::WideString, State)
+UE_TRACE_EVENT_END()
+
+UE_TRACE_EVENT_BEGIN(Cpu, GFP_UpdateStateMachine_EndState, NoSync)
+	UE_TRACE_EVENT_FIELD(UE::Trace::WideString, State)
+UE_TRACE_EVENT_END()
+
 void UGameFeaturesSubsystem::UpdateStateMachine(UGameFeaturePluginStateMachine* This) 
 {
 	LLM_SCOPE_BYTAG(GFP);
@@ -4597,7 +4609,8 @@ void UGameFeaturesSubsystem::UpdateStateMachine(UGameFeaturePluginStateMachine*
 			FName LoaderName = Props.GameFeatureData ? Props.GameFeatureData->GetPackage()->GetFName() : FName(Props.PluginName);
 			UE_TRACK_REFERENCING_PACKAGE_SCOPED(LoaderName, UE::GameFeatures::GetStateName(CurrentState));
 
-			TRACE_CPUPROFILER_EVENT_SCOPE(GFP_UpdateStateMachine_UpdateState);
+			UE_TRACE_LOG_SCOPED_T(Cpu, GFP_UpdateStateMachine_UpdateState, CpuChannel)
+				<< GFP_UpdateStateMachine_UpdateState.State(*UE::GameFeatures::ToString(CurrentState));
 			This->AllStates[CurrentState]->UpdateState(StateStatus);
 		}
 
@@ -4610,7 +4623,8 @@ void UGameFeaturesSubsystem::UpdateStateMachine(UGameFeaturePluginStateMachine*
 		{
 			UE_LOG(LogGameFeatures, Verbose, TEXT("Game feature '%s' transitioning state (%s -> %s)"), *GetGameFeatureName(This), *UE::GameFeatures::ToString(CurrentState), *UE::GameFeatures::ToString(StateStatus.TransitionToState));
 			{
-				TRACE_CPUPROFILER_EVENT_SCOPE(GFP_UpdateStateMachine_EndState);
+				UE_TRACE_LOG_SCOPED_T(Cpu, GFP_UpdateStateMachine_EndState, CpuChannel)
+					<< GFP_UpdateStateMachine_EndState.State(*UE::GameFeatures::ToString(CurrentState));
 				This->AllStates[CurrentState]->EndState();
 				CheckAndCancelBatchingRequestForCurrentState(This);
 			}
@@ -4623,7 +4637,8 @@ void UGameFeaturesSubsystem::UpdateStateMachine(UGameFeaturePluginStateMachine*
 			UE_TRACK_REFERENCING_PACKAGE_SCOPED(LoaderName, UE::GameFeatures::GetStateName(CurrentState));
 
 			{
-				TRACE_CPUPROFILER_EVENT_SCOPE(GFP_UpdateStateMachine_BeginState);
+				UE_TRACE_LOG_SCOPED_T(Cpu, GFP_UpdateStateMachine_BeginState, CpuChannel)
+					<< GFP_UpdateStateMachine_BeginState.State(*UE::GameFeatures::ToString(CurrentState));
 				CheckAddBatchingRequestForCurrentState(This);
 				This->AllStates[CurrentState]->BeginState();
 			}
