CheckingStatus already gathers all the depends and checks their install state
For _Downloading/_Mounting avoid chasing down the install bundle depends for each BeginState for all GFPs

#Synced-CL 49874288

diff --git a/Engine/Plugins/Runtime/GameFeatures/Source/GameFeatures/Private/GameFeaturePluginStateMachine.cpp b/Engine/Plugins/Runtime/GameFeatures/Source/GameFeatures/Private/GameFeaturePluginStateMachine.cpp
--- a/Engine/Plugins/Runtime/GameFeatures/Source/GameFeatures/Private/GameFeaturePluginStateMachine.cpp
+++ b/Engine/Plugins/Runtime/GameFeatures/Source/GameFeatures/Private/GameFeaturePluginStateMachine.cpp
@@ -2059,6 +2059,7 @@ struct FGameFeaturePluginState_Downloading : public FGameFeaturePluginState
 		//without SkipMount and then mount the data, this allows us to pre-download data without mounting it
 		EInstallBundleRequestFlags InstallFlags = StateProperties.ProtocolOptions.GetSubtype<FInstallBundlePluginProtocolOptions>().InstallBundleFlags;
 		InstallFlags |= EInstallBundleRequestFlags::SkipMount;
+		InstallFlags |= EInstallBundleRequestFlags::ExplicitUpdateList; // avoid chasing depends for the install bundles
 		return InstallFlags;
 	}
 
@@ -2449,8 +2450,9 @@ struct FGameFeaturePluginState_Mounting : public FGameFeaturePluginState
 		UE::GameFeatures::GFPSharedInstallTracker.AddBundleRefs(StateProperties.PluginName, UE::GameFeatures::EGFPInstallLevel::Mount, InstallBundles);
 
 		const FInstallBundlePluginProtocolOptions& Options = StateProperties.ProtocolOptions.GetSubtype<FInstallBundlePluginProtocolOptions>();
-		const EInstallBundleRequestFlags InstallFlags = UseAsyncLoading() ?
+		EInstallBundleRequestFlags InstallFlags = UseAsyncLoading() ?
 			(Options.InstallBundleFlags | EInstallBundleRequestFlags::AsyncMount) : Options.InstallBundleFlags;
+		InstallFlags |= EInstallBundleRequestFlags::ExplicitUpdateList; // skip over collecting all the depends here, as we do this in CheckingStatus state
 
 		// Make bundle manager use verbose log level for most logs.
 		// We are already done with downloading, so we don't care about logging too much here unless mounting fails.
