Fix null pointer dereferences of FGameFeaturePluginState when cleaning up subsystem on shutdown.

#Synced-CL 49969322

diff --git a/Engine/Plugins/Runtime/GameFeatures/Source/GameFeatures/Private/GameFeaturePluginStateMachine.cpp b/Engine/Plugins/Runtime/GameFeatures/Source/GameFeatures/Private/GameFeaturePluginStateMachine.cpp
--- a/Engine/Plugins/Runtime/GameFeatures/Source/GameFeatures/Private/GameFeaturePluginStateMachine.cpp
+++ b/Engine/Plugins/Runtime/GameFeatures/Source/GameFeatures/Private/GameFeaturePluginStateMachine.cpp
@@ -4090,6 +4090,15 @@ UGameFeaturePluginStateMachine::UGameFeaturePluginStateMachine(const FObjectInit
 
 }
 
+FDestinationGameFeaturePluginState* AsDestinationState(FGameFeaturePluginState* InState)
+{
+	if (InState)
+	{
+		return InState->AsDestinationState();
+	}
+	return nullptr;
+}
+
 void UGameFeaturesSubsystem::InitStateMachine(UGameFeaturePluginStateMachine* This, FGameFeaturePluginIdentifier InPluginIdentifier, const FGameFeatureProtocolOptions& InProtocolOptions) 
 {
 	LLM_SCOPE_BYTAG(GFP);
@@ -4310,7 +4319,7 @@ void UGameFeaturesSubsystem::RemovePendingTransitionCallback(UGameFeaturePluginS
 		iState < EGameFeaturePluginState::MAX;
 		++iState)
 	{
-		if (FDestinationGameFeaturePluginState* DestState = This->AllStates[iState]->AsDestinationState())
+		if (FDestinationGameFeaturePluginState* DestState = AsDestinationState(This->AllStates[iState].Get()))
 		{
 			if (DestState->OnDestinationStateReached.Remove(InHandle))
 			{
@@ -4326,7 +4335,7 @@ void UGameFeaturesSubsystem::RemovePendingTransitionCallback(UGameFeaturePluginS
 		iState < EGameFeaturePluginState::MAX;
 		++iState)
 	{
-		if (FDestinationGameFeaturePluginState* DestState = This->AllStates[iState]->AsDestinationState())
+		if (FDestinationGameFeaturePluginState* DestState = AsDestinationState(This->AllStates[iState].Get()))
 		{
 			if (DestState->OnDestinationStateReached.RemoveAll(DelegateObject))
 			{
diff --git a/Engine/Plugins/Runtime/GameFeatures/Source/GameFeatures/Private/GameFeaturePluginStateMachine.h b/Engine/Plugins/Runtime/GameFeatures/Source/GameFeatures/Private/GameFeaturePluginStateMachine.h
--- a/Engine/Plugins/Runtime/GameFeatures/Source/GameFeatures/Private/GameFeaturePluginStateMachine.h
+++ b/Engine/Plugins/Runtime/GameFeatures/Source/GameFeatures/Private/GameFeaturePluginStateMachine.h
@@ -497,6 +497,8 @@ class UGameFeaturePluginStateMachine : public UObject
 public:
 	UGameFeaturePluginStateMachine(const FObjectInitializer& ObjectInitializer);
 
+	FDestinationGameFeaturePluginState* AsDestinationState(FGameFeaturePluginState* InState);
+
 	/** Information about the current state */
 	FGameFeaturePluginStateInfo CurrentStateInfo;
 
