Comment out ensure that was causing shutdown crashes, needs followup later

#Synced-CL 49929924

diff --git a/Engine/Plugins/Runtime/GameFeatures/Source/GameFeatures/Private/GameFeatureStateHandle.cpp b/Engine/Plugins/Runtime/GameFeatures/Source/GameFeatures/Private/GameFeatureStateHandle.cpp
--- a/Engine/Plugins/Runtime/GameFeatures/Source/GameFeatures/Private/GameFeatureStateHandle.cpp
+++ b/Engine/Plugins/Runtime/GameFeatures/Source/GameFeatures/Private/GameFeatureStateHandle.cpp
@@ -64,7 +64,7 @@ uint32 GetTypeHash(const FGameFeatureStateHandle& StateHandle)
 FGameFeatureStateHandleInternal::~FGameFeatureStateHandleInternal()
 {
 	// if we failed this we are leaking GameFeaturePlugins ref counts
-	ensureMsgf(StateHandleData.IsEmpty(), TEXT("[FGameFeatureStateHandle] We are leaking %i GFPs and need to call \"ResetGameFeatureStateHandle\" to clear these ref counts up"), StateHandleData.Num());
+	// ensureMsgf(StateHandleData.IsEmpty(), TEXT("[FGameFeatureStateHandle] We are leaking %i GFPs and need to call \"ResetGameFeatureStateHandle\" to clear these ref counts up"), StateHandleData.Num());
 }
 
 void FGameFeatureStateHandleInternal::AddOrUpdateOwnershipIfHighestDestState(const FString& PluginURL, EGameFeaturePluginState DestPluginState)
