When a GFP has a BuiltInAutoState of active dont allow them to be moved back to loaded

diff --git a/Engine/Plugins/Runtime/GameFeatures/Source/GameFeatures/Private/GameFeaturesSubsystem.cpp b/Engine/Plugins/Runtime/GameFeatures/Source/GameFeatures/Private/GameFeaturesSubsystem.cpp
--- a/Engine/Plugins/Runtime/GameFeatures/Source/GameFeatures/Private/GameFeaturesSubsystem.cpp
+++ b/Engine/Plugins/Runtime/GameFeatures/Source/GameFeatures/Private/GameFeaturesSubsystem.cpp
@@ -4157,6 +4157,18 @@ bool UGameFeaturesSubsystem::FindPluginDependencyStateMachinesToActivate(const F
 bool UGameFeaturesSubsystem::FindPluginDependencyStateMachinesToDeactivate(const FString& PluginURL, const FString& PluginFilename, TArray<UGameFeaturePluginStateMachine*>& OutDependencyMachines) const
 {
 	const bool bResult = EnumeratePluginDependenciesWithShouldActivate(PluginURL, PluginFilename, [this, &OutDependencyMachines](const FString& DependencyName, const FString& DependencyURL) {
+		FStringView PluginPath;
+		if (UGameFeaturesSubsystem::ParsePluginURL(DependencyURL, nullptr, &PluginPath))
+		{
+			if (TSharedPtr<FGameFeaturePluginDetails> Details = GetGameFeaturePluginDetailsInternal(FString(PluginPath)))
+			{
+				if (Details->BuiltInAutoState >= EBuiltInAutoState::Active) // if our builtin state is active dont deactivate these
+				{
+					return true;
+				}
+			}
+		}
+
 		UGameFeaturePluginStateMachine* Dependency = FindGameFeaturePluginStateMachine(DependencyURL);
 		if (Dependency)
 		{
