Skip validation for plugins not under the Engine or the Project Plugins directory.

#Synced-CL 50152832

diff --git a/Engine/Plugins/Runtime/GameFeatures/Source/GameFeaturesEditor/Private/IllegalPluginDependenciesValidator.cpp b/Engine/Plugins/Runtime/GameFeatures/Source/GameFeaturesEditor/Private/IllegalPluginDependenciesValidator.cpp
--- a/Engine/Plugins/Runtime/GameFeatures/Source/GameFeaturesEditor/Private/IllegalPluginDependenciesValidator.cpp
+++ b/Engine/Plugins/Runtime/GameFeatures/Source/GameFeaturesEditor/Private/IllegalPluginDependenciesValidator.cpp
@@ -79,6 +79,11 @@ EDataValidationResult UIllegalPluginDependenciesValidator::ValidateLoadedAsset_I
 				{
 					continue;
 				}
+				if (!FPaths::IsUnderDirectory(Plugin->GetDescriptorFileName(), FPaths::EnginePluginsDir()) && !FPaths::IsUnderDirectory(Plugin->GetDescriptorFileName(), FPaths::ProjectPluginsDir()))
+				{
+					// do not validate plugins outside of this project.
+					continue;
+				}
 
 				// Plugin is not a GFP so it can not have any GFP dependencies
 				const FPluginDescriptor& Descriptor = Plugin->GetDescriptor();
