GFP refactor - adding new executor.
Route UpdateStateMachineImmediate and UpdateStateMachineDeferred through new executor for scheduling.
Precursor to refactoring state updates to all operate in bulk on all GFPs making the same state transition.

#Synced-CL 50193578

diff --git a/Engine/Plugins/Runtime/GameFeatures/Source/GameFeatures/Private/GameFeaturePluginStateMachine.cpp b/Engine/Plugins/Runtime/GameFeatures/Source/GameFeatures/Private/GameFeaturePluginStateMachine.cpp
--- a/Engine/Plugins/Runtime/GameFeatures/Source/GameFeatures/Private/GameFeaturePluginStateMachine.cpp
+++ b/Engine/Plugins/Runtime/GameFeatures/Source/GameFeatures/Private/GameFeaturePluginStateMachine.cpp
@@ -2,6 +2,7 @@
 
 #include "GameFeaturePluginStateMachine.h"
 
+#include "Algo/RandomShuffle.h"
 #include "AssetRegistry/AssetRegistryState.h"
 #include "Engine/StreamableManager.h"
 #include "GameFeatureData.h"
@@ -25,6 +26,7 @@
 #include "Misc/ConfigCacheIni.h"
 #include "Misc/ConfigUtilities.h"
 #include "Misc/CoreDelegates.h"
+#include "Misc/CoreMiscDefines.h"
 #include "Misc/EnumRange.h"
 #include "Misc/FileHelper.h"
 #include "Misc/PackageName.h"
@@ -34,12 +36,14 @@
 #include "Algo/AllOf.h"
 #include "Algo/Find.h"
 #include "Misc/TVariantMeta.h"
+#include "ProfilingDebugging/CpuProfilerTrace.h"
 #include "RenderDeferredCleanup.h"
 #include "Serialization/MemoryReader.h"
 #include "String/ParseTokens.h"
 #include "String/LexFromString.h"
 #include "Tasks/Pipe.h"
 #include "GameFeaturesProjectPolicies.h"
+#include "Templates/UnrealTemplate.h"
 #include "UObject/NameTypes.h"
 #include "UObject/ObjectRename.h"
 #include "UObject/ReferenceChainSearch.h"
@@ -72,6 +76,42 @@ namespace UE::GameFeatures
 
 	static UE::GameFeatures::FResult CanceledResult = MakeError(StateMachineErrorNamespace + TEXT("Canceled"));
 
+	bool bUseNewExecutor = false;
+	static FAutoConsoleVariableRef CVarUseNewExecutor(TEXT("GameFeaturePlugin.UseNewExecutor"),
+		bUseNewExecutor,
+		TEXT("Should control of state machine updates be handled by the GameFeaturesSubsystem (new behavior) to all more batch & parallel processing."),
+	ECVF_ReadOnly | ECVF_SaveForNextBoot);
+
+	// Options for how to defer execution of state updates and how to order those updates
+	struct FExecutionOptions
+	{
+		// Group by current state before executing
+		bool bGrouped = false;
+		// Execute the update requests (potentially within each group) in a random order
+		bool bShuffle = false;
+		// Reverse the order of update requests so that deferred request are executed in LIFO order
+		bool bReverse = false;
+	} ExecutionOptions;
+
+	static FAutoConsoleVariableRef CVarExecutionGrouped(TEXT("GameFeaturePlugin.Execution.Grouped"),
+		ExecutionOptions.bGrouped,
+		TEXT("Group deferred updates by current state before executing. This is the intended future default path."),
+		ECVF_Default);
+	static FAutoConsoleVariableRef CVarExecutionShuffle(TEXT("GameFeaturePlugin.Execution.Shuffle"),
+		ExecutionOptions.bShuffle,
+		TEXT("Shuffle deferred updates before execution (within a group if .Grouped is set). For testing implicit dependencies."),
+		ECVF_Default);
+	static FAutoConsoleVariableRef CVarExecutionReverse(TEXT("GameFeaturePlugin.Execution.Reverse"),
+		ExecutionOptions.bReverse,
+		TEXT("Reverse deferred updates before execution to execute in LIFO order. For testing deferred dependencies."),
+		ECVF_Default);
+
+	bool bLoadBuiltin_DeferExecution = false;
+	static FAutoConsoleVariableRef CVarLoadBuiltin_DeferExecution (TEXT("GameFeaturePlugin.LoadBuiltins.DeferExecution"),
+		bLoadBuiltin_DeferExecution ,
+		TEXT("When GameFeaturePlugin.UseNewExecutor is true, should execution of state machine updates in LoadBuiltInGameFeaturePlugins be deferred until all state transitions have been queued?"),
+	ECVF_ReadOnly);
+
 	static int32 ShouldLogMountedFiles = 0;
 	static FAutoConsoleVariableRef CVarShouldLogMountedFiles(TEXT("GameFeaturePlugin.ShouldLogMountedFiles"),
 		ShouldLogMountedFiles,
@@ -153,7 +193,7 @@ namespace UE::GameFeatures
 							StateMachine && ensureAlways(&StateMachine->StateProperties == &StateProperties))
 						{
 							StateProperties.bIsLoadingLocalizationData = false;
-							StateProperties.OnRequestUpdateStateMachine.ExecuteIfBound();
+							StateProperties.UpdateStateMachineImmediate();
 						}
 					});
 				}
@@ -216,16 +256,17 @@ namespace UE::GameFeatures
 				FTSTicker::RemoveTicker(MoveTemp(TickHandle));
 			}
 
-			FGameFeaturePluginRequestUpdateStateMachine UpdateRequest;
+			int64 UpdateRequest = 0;
+			UGameFeaturesSubsystem& Subsystem = UGameFeaturesSubsystem::Get();
 			while (UpdateRequests.Dequeue(UpdateRequest))
 			{
-				UpdateRequest.ExecuteIfBound();
+				Subsystem.TryUpdateStateMachineImmediate(UpdateRequest);
 			}
 		}
 
-		void AddUpdateRequest(FGameFeaturePluginRequestUpdateStateMachine UpdateRequest)
+		void AddUpdateRequest(FGameFeaturePluginStateMachineProperties StateProperties)
 		{
-			UpdateRequests.Enqueue(MoveTemp(UpdateRequest));
+			UpdateRequests.Enqueue(StateProperties.Id);
 			EnableTick();
 		}
 
@@ -248,10 +289,11 @@ namespace UE::GameFeatures
 				constexpr double AllottedTime = MaxFrameTime / 2;
 				const double StartTime = FPlatformTime::Seconds();
 
-				FGameFeaturePluginRequestUpdateStateMachine UpdateRequest;
+				UGameFeaturesSubsystem& Subsystem = UGameFeaturesSubsystem::Get();
+				int64 UpdateRequest = 0;
 				while (UpdateRequests.Dequeue(UpdateRequest))
 				{
-					UpdateRequest.ExecuteIfBound();
+					Subsystem.TryUpdateStateMachineImmediate(UpdateRequest);
 
 					const double ElapsedTime = FPlatformTime::Seconds() - StartTime;
 					if ((ElapsedTime > AllottedTime) || ((DeltaTime + ElapsedTime) > MaxFrameTime))
@@ -273,7 +315,7 @@ namespace UE::GameFeatures
 		}
 
 	private:
-		TQueue<FGameFeaturePluginRequestUpdateStateMachine> UpdateRequests;
+		TQueue<int64> UpdateRequests;
 		FTSTicker::FDelegateHandle TickHandle;
 	};
 
@@ -285,6 +327,10 @@ namespace UE::GameFeatures
 		FConsoleVariableDelegate::CreateLambda(
 			[](IConsoleVariable* Var)
 			{
+				if (bUseNewExecutor)
+				{
+					return;
+				}
 				if (Var->GetBool())
 				{
 					if (!RealtimeMode)
@@ -301,7 +347,7 @@ namespace UE::GameFeatures
 			ECVF_ReadOnly);
 
 #if WITH_EDITOR
-	TMap<FString, FGameFeaturePluginRequestUpdateStateMachine> PluginsToUnloadAssets;
+	TMap<FString, int64> PluginsToUnloadAssets;
 	FTSTicker::FDelegateHandle UnloadPluginAssetsHandle;
 
 	bool TickUnloadPluginAssets(float /*DeltaTime*/)
@@ -309,33 +355,34 @@ namespace UE::GameFeatures
 		UnloadPluginAssetsHandle.Reset();
 	
 		TArray<FString> PluginNames;
-		TArray<FGameFeaturePluginRequestUpdateStateMachine> UpdateStateMachineDelegates;
+		TArray<int64> UpdateStateMachineIds;
 		{
 			PluginNames.Reserve(PluginsToUnloadAssets.Num());
-			UpdateStateMachineDelegates.Reserve(PluginsToUnloadAssets.Num());
-			for (TPair<FString, FGameFeaturePluginRequestUpdateStateMachine>& PluginsToUnloadAsset : PluginsToUnloadAssets)
+			UpdateStateMachineIds.Reserve(PluginsToUnloadAssets.Num());
+			for (TPair<FString, int64>& PluginsToUnloadAsset : PluginsToUnloadAssets)
 			{
 				PluginNames.Add(PluginsToUnloadAsset.Key);
-				UpdateStateMachineDelegates.Add(MoveTemp(PluginsToUnloadAsset.Value));
+				UpdateStateMachineIds.Add(PluginsToUnloadAsset.Value);
 			}
 			PluginsToUnloadAssets.Empty();
 		}
 
 		verify(FPluginUtils::UnloadPluginsAssets(PluginNames));
 
-		for (const FGameFeaturePluginRequestUpdateStateMachine& UpdateStateMachineDelegate : UpdateStateMachineDelegates)
+		UGameFeaturesSubsystem& Subsystem = UGameFeaturesSubsystem::Get();
+		for (int64 Id : UpdateStateMachineIds)
 		{
-			UpdateStateMachineDelegate.ExecuteIfBound();
+			Subsystem.TryUpdateStateMachineImmediate(Id);
 		}
 
 		return false;
 	}
 
-	void ScheduleUnloadPluginAssets(const FString& PluginName, const FGameFeaturePluginRequestUpdateStateMachine& UpdateStateMachineDelegate)
+	void ScheduleUnloadPluginAssets(const FString& PluginName, const FGameFeaturePluginStateMachineProperties& StateProperties)
 	{
 		check(IsInGameThread());
 		ensure(!PluginsToUnloadAssets.Contains(PluginName));
-		PluginsToUnloadAssets.Add(PluginName, UpdateStateMachineDelegate);
+		PluginsToUnloadAssets.Add(PluginName, StateProperties.Id);
 		if (!UnloadPluginAssetsHandle.IsValid())
 		{
 			UnloadPluginAssetsHandle = FTSTicker::GetCoreTicker().AddTicker(FTickerDelegate::CreateStatic(&TickUnloadPluginAssets));
@@ -558,6 +605,170 @@ namespace UE::GameFeatures
 
 		return FName(TEXT("Unknown"));
 	};
+
+	// Default next state to transition to when transitioning down towards Terminal
+	static constexpr EGameFeaturePluginState NextStateTowardsTerminal[EGameFeaturePluginState::MAX] = 
+	{
+		/* Uninitialized */ 				EGameFeaturePluginState::Uninitialized,
+		/* Terminal */ 						EGameFeaturePluginState::Terminal,
+		/* UnknownStatus */ 				EGameFeaturePluginState::Terminal,
+		/* Uninstalled */ 					EGameFeaturePluginState::Terminal,
+		/* Uninstalling */					EGameFeaturePluginState::Uninstalled,
+		/* ErrorUninstalling */ 			EGameFeaturePluginState::Uninstalling, // NOTE: May not need to uninstall
+		/* CheckingStatus */ 				EGameFeaturePluginState::StatusKnown, // NOTE: No down transition
+		/* ErrorCheckingStatus */ 			EGameFeaturePluginState::Terminal,
+		/* ErrorUnavailable */ 				EGameFeaturePluginState::Terminal,
+		/* StatusKnown */ 					EGameFeaturePluginState::Uninstalling, // NOTE: May not need to uninstall
+		/* Releasing */ 					EGameFeaturePluginState::StatusKnown,
+		/* ErrorManagingData */ 			EGameFeaturePluginState::Releasing, // TODO: Confirm
+		/* Downloading */ 					EGameFeaturePluginState::Installed, // NOTE: No down transition
+		/* Installed */ 					EGameFeaturePluginState::Releasing,
+		/* ErrorMounting */ 				EGameFeaturePluginState::Unmounting,
+		/* ErrorWaitingForDependencies */	EGameFeaturePluginState::Unmounting,
+		/* ErrorRegistering */		 		EGameFeaturePluginState::Unregistering,
+		/* WaitingForDependencies */ 		EGameFeaturePluginState::AssetDependencyStreaming, // NOTE: No down transition
+		/* AssetDependencyStreamOut */	 	EGameFeaturePluginState::Unmounting,
+		/* ErrorAssetDependencyStreaming */ EGameFeaturePluginState::AssetDependencyStreamOut,
+		/* AssetDependencyStreaming */ 		EGameFeaturePluginState::Registering, // NOTE: No down transition
+		/* Unmounting */ 					EGameFeaturePluginState::Installed,
+		/* Mounting */ 						EGameFeaturePluginState::WaitingForDependencies, // NOTE: No down transition
+		/* Unregistering */ 				EGameFeaturePluginState::AssetDependencyStreamOut,
+		/* Registering */ 					EGameFeaturePluginState::Registered, // NOTE: No down transition
+		/* Registered */ 					EGameFeaturePluginState::Unregistering,
+		/* ErrorLoading */					EGameFeaturePluginState::Unloading,
+		/* Unloading */ 					EGameFeaturePluginState::Registered,
+		/* Loading */ 						EGameFeaturePluginState::Loaded, // NOTE: No down transitoin
+		/* Loaded */ 						EGameFeaturePluginState::Unloading,
+		/* ErrorActivatingDependencies */	EGameFeaturePluginState::Unloading,
+		/* ActivatingDependencies */ 		EGameFeaturePluginState::Activating, // NOTE: No down transition
+		/* Deactivating */ 					EGameFeaturePluginState::Loaded,
+		/* ErrorDeactivatingDependencies */	EGameFeaturePluginState::DeactivatingDependencies,
+		/* DeactivatingDependencies */ 		EGameFeaturePluginState::Deactivating,
+		/* Activating */ 					EGameFeaturePluginState::Active, // NOTE: No down transition
+		/* Active */ 						EGameFeaturePluginState::DeactivatingDependencies,
+	};
+	// Default next state to transition to when transitioning up towards Active
+	static constexpr EGameFeaturePluginState NextStateTowardsActive[EGameFeaturePluginState::MAX] = 
+	{
+		/* Uninitialized */ 				EGameFeaturePluginState::UnknownStatus,
+		/* Terminal */ 						EGameFeaturePluginState::Terminal, // No up transition, must be recreated
+		/* UnknownStatus */ 				EGameFeaturePluginState::CheckingStatus,
+		/* Uninstalled */ 					EGameFeaturePluginState::CheckingStatus,
+		/* Uninstalling */					EGameFeaturePluginState::Uninstalled, // NOTE: No up transition
+		/* ErrorUninstalling */ 			EGameFeaturePluginState::StatusKnown,
+		/* CheckingStatus */ 				EGameFeaturePluginState::StatusKnown,
+		/* ErrorCheckingStatus */ 			EGameFeaturePluginState::CheckingStatus,
+		/* ErrorUnavailable */ 				EGameFeaturePluginState::CheckingStatus,
+		/* StatusKnown */ 					EGameFeaturePluginState::Downloading, // NOTE: May go straight to Installed
+		/* Releasing */ 					EGameFeaturePluginState::StatusKnown, // NOTE: No up transition
+		/* ErrorManagingData */ 			EGameFeaturePluginState::Downloading,
+		/* Downloading */ 					EGameFeaturePluginState::Installed,
+		/* Installed */ 					EGameFeaturePluginState::Mounting,
+		/* ErrorMounting */ 				EGameFeaturePluginState::Mounting,
+		/* ErrorWaitingForDependencies */	EGameFeaturePluginState::WaitingForDependencies,
+		/* ErrorRegistering */		 		EGameFeaturePluginState::Registering,
+		/* WaitingForDependencies */ 		EGameFeaturePluginState::AssetDependencyStreaming,
+		/* AssetDependencyStreamOut */	 	EGameFeaturePluginState::Unmounting, // NOTE: No up transition
+		/* ErrorAssetDependencyStreaming */ EGameFeaturePluginState::AssetDependencyStreaming,
+		/* AssetDependencyStreaming */ 		EGameFeaturePluginState::Registering,
+		/* Unmounting */ 					EGameFeaturePluginState::Installed, // NOTE: No up transition
+		/* Mounting */ 						EGameFeaturePluginState::WaitingForDependencies,
+		/* Unregistering */ 				EGameFeaturePluginState::AssetDependencyStreamOut, // NOTE: No up transition
+		/* Registering */ 					EGameFeaturePluginState::Registered,
+		/* Registered */ 					EGameFeaturePluginState::Loading,
+		/* ErrorLoading */					EGameFeaturePluginState::Loading,
+		/* Unloading */ 					EGameFeaturePluginState::Registered, // NOTE: No up transition
+		/* Loading */ 						EGameFeaturePluginState::Loaded,
+		/* Loaded */ 						EGameFeaturePluginState::ActivatingDependencies,
+		/* ErrorActivatingDependencies */	EGameFeaturePluginState::ActivatingDependencies,
+		/* ActivatingDependencies */ 		EGameFeaturePluginState::Activating,
+		/* Deactivating */ 					EGameFeaturePluginState::Loaded, // NOTE: No up transition
+		/* ErrorDeactivatingDependencies */	EGameFeaturePluginState::DeactivatingDependencies,
+		/* DeactivatingDependencies */ 		EGameFeaturePluginState::Deactivating, // NOTE: No up transition
+		/* Activating */ 					EGameFeaturePluginState::Active,
+		/* Active */ 						EGameFeaturePluginState::Active
+	};
+	// Default next state to transition to when encountering an error
+	static constexpr EGameFeaturePluginState NextStateOnError[EGameFeaturePluginState::MAX] = 
+	{
+		/* Uninitialized */ 				EGameFeaturePluginState::Uninitialized,// No errors handled
+		/* Terminal */ 						EGameFeaturePluginState::Terminal, // No errors handled
+		/* UnknownStatus */ 				EGameFeaturePluginState::UnknownStatus, // No errors handled
+		/* Uninstalled */ 					EGameFeaturePluginState::Uninstalled, // No errors handled
+		/* Uninstalling */					EGameFeaturePluginState::ErrorUninstalling, 
+		/* ErrorUninstalling */ 			EGameFeaturePluginState::ErrorUninstalling, // Already an error state
+		/* CheckingStatus */ 				EGameFeaturePluginState::ErrorCheckingStatus,
+		/* ErrorCheckingStatus */ 			EGameFeaturePluginState::ErrorCheckingStatus, // Already an error state
+		/* ErrorUnavailable */ 				EGameFeaturePluginState::ErrorUnavailable, // Already an error state
+		/* StatusKnown */ 					EGameFeaturePluginState::StatusKnown, // No errors handled
+		/* Releasing */ 					EGameFeaturePluginState::ErrorManagingData, 
+		/* ErrorManagingData */ 			EGameFeaturePluginState::ErrorManagingData, // Already an error state
+		/* Downloading */ 					EGameFeaturePluginState::ErrorManagingData,
+		/* Installed */ 					EGameFeaturePluginState::Installed, // No errors handled
+		/* ErrorMounting */ 				EGameFeaturePluginState::ErrorMounting, // Already an error state
+		/* ErrorWaitingForDependencies */	EGameFeaturePluginState::ErrorWaitingForDependencies, // Already an error state
+		/* ErrorRegistering */		 		EGameFeaturePluginState::ErrorRegistering, // Already an error state
+		/* WaitingForDependencies */ 		EGameFeaturePluginState::ErrorWaitingForDependencies,
+		/* AssetDependencyStreamOut */	 	EGameFeaturePluginState::ErrorAssetDependencyStreaming, 
+		/* ErrorAssetDependencyStreaming */ EGameFeaturePluginState::ErrorAssetDependencyStreaming, // Already an error state
+		/* AssetDependencyStreaming */ 		EGameFeaturePluginState::ErrorAssetDependencyStreaming,
+		/* Unmounting */ 					EGameFeaturePluginState::ErrorMounting,
+		/* Mounting */ 						EGameFeaturePluginState::ErrorMounting,
+		/* Unregistering */ 				EGameFeaturePluginState::AssetDependencyStreamOut, // NOTE: No up transition
+		/* Registering */ 					EGameFeaturePluginState::ErrorRegistering,
+		/* Registered */ 					EGameFeaturePluginState::Registered, // No errors handled
+		/* ErrorLoading */					EGameFeaturePluginState::Loading,
+		/* Unloading */ 					EGameFeaturePluginState::ErrorLoading, 
+		/* Loading */ 						EGameFeaturePluginState::ErrorLoading,
+		/* Loaded */ 						EGameFeaturePluginState::Loaded, // No errors handled
+		/* ErrorActivatingDependencies */	EGameFeaturePluginState::ErrorActivatingDependencies, // Already an error state
+		/* ActivatingDependencies */ 		EGameFeaturePluginState::ErrorActivatingDependencies,
+		/* Deactivating */ 					EGameFeaturePluginState::Deactivating, // No errors handled
+		/* ErrorDeactivatingDependencies */	EGameFeaturePluginState::ErrorActivatingDependencies, // Already an error state
+		/* DeactivatingDependencies */ 		EGameFeaturePluginState::ErrorDeactivatingDependencies, 
+		/* Activating */ 					EGameFeaturePluginState::Activating, // No errors handled
+		/* Active */ 						EGameFeaturePluginState::Active // No errors handled
+	};
+	static constexpr EGameFeaturePluginStateType StateType[EGameFeaturePluginState::MAX] = 
+	{
+		/* Uninitialized */ 				EGameFeaturePluginStateType::Transition,
+		/* Terminal */ 						EGameFeaturePluginStateType::Destination,
+		/* UnknownStatus */ 				EGameFeaturePluginStateType::Destination,
+		/* Uninstalled */ 					EGameFeaturePluginStateType::Destination,
+		/* Uninstalling */					EGameFeaturePluginStateType::Transition,
+		/* ErrorUninstalling */ 			EGameFeaturePluginStateType::Error,
+		/* CheckingStatus */ 				EGameFeaturePluginStateType::Transition,
+		/* ErrorCheckingStatus */ 			EGameFeaturePluginStateType::Error,
+		/* ErrorUnavailable */ 				EGameFeaturePluginStateType::Error,
+		/* StatusKnown */ 					EGameFeaturePluginStateType::Destination,
+		/* Releasing */ 					EGameFeaturePluginStateType::Transition,
+		/* ErrorManagingData */ 			EGameFeaturePluginStateType::Error,
+		/* Downloading */ 					EGameFeaturePluginStateType::Transition,
+		/* Installed */ 					EGameFeaturePluginStateType::Destination,
+		/* ErrorMounting */ 				EGameFeaturePluginStateType::Error,
+		/* ErrorWaitingForDependencies */	EGameFeaturePluginStateType::Error,
+		/* ErrorRegistering */		 		EGameFeaturePluginStateType::Error,
+		/* WaitingForDependencies */ 		EGameFeaturePluginStateType::Transition,
+		/* AssetDependencyStreamOut */	 	EGameFeaturePluginStateType::Transition,
+		/* ErrorAssetDependencyStreaming */ EGameFeaturePluginStateType::Error,
+		/* AssetDependencyStreaming */ 		EGameFeaturePluginStateType::Transition,
+		/* Unmounting */ 					EGameFeaturePluginStateType::Transition,
+		/* Mounting */ 						EGameFeaturePluginStateType::Transition,
+		/* Unregistering */ 				EGameFeaturePluginStateType::Transition,
+		/* Registering */ 					EGameFeaturePluginStateType::Transition,
+		/* Registered */ 					EGameFeaturePluginStateType::Destination,
+		/* ErrorLoading */					EGameFeaturePluginStateType::Error,
+		/* Unloading */ 					EGameFeaturePluginStateType::Transition,
+		/* Loading */ 						EGameFeaturePluginStateType::Transition,
+		/* Loaded */ 						EGameFeaturePluginStateType::Destination,
+		/* ErrorActivatingDependencies */	EGameFeaturePluginStateType::Error,
+		/* ActivatingDependencies */ 		EGameFeaturePluginStateType::Transition,
+		/* Deactivating */ 					EGameFeaturePluginStateType::Transition,
+		/* ErrorDeactivatingDependencies */	EGameFeaturePluginStateType::Error,
+		/* DeactivatingDependencies */ 		EGameFeaturePluginStateType::Transition,
+		/* Activating */ 					EGameFeaturePluginStateType::Transition,
+		/* Active */ 						EGameFeaturePluginStateType::Destination,
+	};
 }
 
 void FGameFeaturePluginStateStatus::SetTransition(EGameFeaturePluginState InTransitionToState)
@@ -664,6 +875,11 @@ FErrorGameFeaturePluginState* FGameFeaturePluginState::AsErrorState()
 
 void FGameFeaturePluginState::UpdateStateMachineDeferred(float Delay /*= 0.0f*/) const
 {
+	// If the shared executor is enabled, don't create an individual callback for this state
+	if (StateProperties.UpdateStateMachineDeferred(Delay))
+	{
+		return;
+	}
 	CleanupDeferredUpdateCallbacks();
 
 	GFP_STATE_LOG(VeryVerbose, "queueing deferred update");
@@ -671,26 +887,19 @@ void FGameFeaturePluginState::UpdateStateMachineDeferred(float Delay /*= 0.0f*/)
 	{
 		// @note Release FGameFeaturePluginState::TickHandle first in case the termination callback triggers a GC and destroys the state machine
 		TickHandle.Reset();
-		if (StateProperties.OnRequestUpdateStateMachine.IsBound())
-		{
-			StateProperties.OnRequestUpdateStateMachine.Execute();
-		}
-		else
-		{
-			GFP_STATE_LOG(Warning, "OnRequestUpdateStateMachine not bound");
-		}
+		StateProperties.UpdateStateMachineImmediate();
 		return false;
 	}), Delay);
 }
 
 void FGameFeaturePluginState::UpdateStateMachineImmediate() const
 {
-	StateProperties.OnRequestUpdateStateMachine.ExecuteIfBound();
+	StateProperties.UpdateStateMachineImmediate();
 }
 
 void FGameFeaturePluginState::UpdateProgress(float Progress) const
 {
-	StateProperties.OnFeatureStateProgressUpdate.ExecuteIfBound(Progress);
+	StateProperties.StateProgress = Progress; 
 }
 
 bool FGameFeaturePluginState::CanBatchProcess() const
@@ -766,6 +975,26 @@ bool FGameFeaturePluginState::UseAsyncLoading() const
 	return AllowAsyncLoading() && UE::GameFeatures::CVarAsyncLoad.GetValueOnGameThread();
 }
 
+void FDestinationGameFeaturePluginState::UpdateState(FGameFeaturePluginStateStatus& StateStatus)
+{
+	if (StateProperties.Destination < StateProperties.CurrentState)
+	{
+		StateStatus.SetTransition(UE::GameFeatures::NextStateTowardsTerminal[StateProperties.CurrentState]);
+	}
+	else if (StateProperties.Destination > StateProperties.CurrentState)
+	{
+		StateStatus.SetTransition(UE::GameFeatures::NextStateTowardsActive[StateProperties.CurrentState]);
+	}
+}
+
+void FDestinationGameFeaturePluginState::TryCancelState()
+{
+}
+
+void FDestinationGameFeaturePluginState::EndState()
+{
+}
+
 /*
 =========================================================
   States
@@ -800,7 +1029,7 @@ struct FTransitionDependenciesGameFeaturePluginState : public FGameFeaturePlugin
 			bCheckedRealtimeMode = true;
 			if (UE::GameFeatures::RealtimeMode)
 			{
-				UE::GameFeatures::RealtimeMode->AddUpdateRequest(StateProperties.OnRequestUpdateStateMachine);
+				UE::GameFeatures::RealtimeMode->AddUpdateRequest(StateProperties);
 				return;
 			}
 		}
@@ -1101,19 +1330,8 @@ struct FGameFeaturePluginState_UnknownStatus : public FDestinationGameFeaturePlu
 {
 	using FDestinationGameFeaturePluginState::FDestinationGameFeaturePluginState;
 
-	virtual void UpdateState(FGameFeaturePluginStateStatus& StateStatus) override
-	{
-		if (StateProperties.Destination < EGameFeaturePluginState::UnknownStatus)
-		{
-			StateStatus.SetTransition(EGameFeaturePluginState::Terminal);
-		}
-		else if (StateProperties.Destination > EGameFeaturePluginState::UnknownStatus)
-		{
-			StateStatus.SetTransition(EGameFeaturePluginState::CheckingStatus);
-
-			UGameFeaturesSubsystem::Get().OnGameFeatureCheckingStatus(StateProperties.PluginIdentifier);
-		}
-	}
+	static_assert(UE::GameFeatures::NextStateTowardsTerminal[EGameFeaturePluginState::UnknownStatus] == EGameFeaturePluginState::Terminal);
+	static_assert(UE::GameFeatures::NextStateTowardsActive[EGameFeaturePluginState::UnknownStatus] == EGameFeaturePluginState::CheckingStatus);
 };
 
 struct FGameFeaturePluginState_CheckingStatus : public FGameFeaturePluginState
@@ -1125,6 +1343,8 @@ struct FGameFeaturePluginState_CheckingStatus : public FGameFeaturePluginState
 
 	virtual void BeginState() override
 	{
+		UGameFeaturesSubsystem::Get().OnGameFeatureCheckingStatus(StateProperties.PluginIdentifier);
+
 		bParsedURL = false;
 		bIsAvailable = false;
 	}
@@ -1213,106 +1433,39 @@ struct FGameFeaturePluginState_ErrorCheckingStatus : public FErrorGameFeaturePlu
 {
 	using FErrorGameFeaturePluginState::FErrorGameFeaturePluginState;
 
-	virtual void UpdateState(FGameFeaturePluginStateStatus& StateStatus) override
-	{
-		if (StateProperties.Destination < EGameFeaturePluginState::ErrorCheckingStatus)
-		{
-			StateStatus.SetTransition(EGameFeaturePluginState::Terminal);
-		}
-		else if (StateProperties.Destination > EGameFeaturePluginState::ErrorCheckingStatus)
-		{
-			StateStatus.SetTransition(EGameFeaturePluginState::CheckingStatus);
-		}
-	}
+	static_assert(UE::GameFeatures::NextStateTowardsTerminal[EGameFeaturePluginState::ErrorCheckingStatus] == EGameFeaturePluginState::Terminal);
+	static_assert(UE::GameFeatures::NextStateTowardsActive[EGameFeaturePluginState::ErrorCheckingStatus] == EGameFeaturePluginState::CheckingStatus);
 };
 
 struct FGameFeaturePluginState_ErrorUnavailable : public FErrorGameFeaturePluginState
 {
 	using FErrorGameFeaturePluginState::FErrorGameFeaturePluginState;
 
-	virtual void UpdateState(FGameFeaturePluginStateStatus& StateStatus) override
-	{
-		if (StateProperties.Destination < EGameFeaturePluginState::ErrorUnavailable)
-		{
-			StateStatus.SetTransition(EGameFeaturePluginState::Terminal);
-		}
-		else if (StateProperties.Destination > EGameFeaturePluginState::ErrorUnavailable)
-		{
-			StateStatus.SetTransition(EGameFeaturePluginState::CheckingStatus);
-		}
-	}
+	static_assert(UE::GameFeatures::NextStateTowardsTerminal[EGameFeaturePluginState::ErrorUnavailable] == EGameFeaturePluginState::Terminal);
+	static_assert(UE::GameFeatures::NextStateTowardsActive[EGameFeaturePluginState::ErrorUnavailable] == EGameFeaturePluginState::CheckingStatus);
 };
 
 struct FGameFeaturePluginState_StatusKnown : public FDestinationGameFeaturePluginState
 {
 	using FDestinationGameFeaturePluginState::FDestinationGameFeaturePluginState;
 
-	virtual void UpdateState(FGameFeaturePluginStateStatus& StateStatus) override
-	{
-		if (StateProperties.Destination < EGameFeaturePluginState::StatusKnown)
-		{
-			if (ShouldVisitUninstallStateBeforeTerminal())
-			{
-				StateStatus.SetTransition(EGameFeaturePluginState::Uninstalling);
-			}
-			else
-			{
-				StateStatus.SetTransition(EGameFeaturePluginState::Terminal);
-			}
-		}
-		else if (StateProperties.Destination > EGameFeaturePluginState::StatusKnown)
-		{
-			if (StateProperties.GetPluginProtocol() != EGameFeaturePluginProtocol::File)
-			{
-				StateStatus.SetTransition(EGameFeaturePluginState::Downloading);
-			}
-			else
-			{
-				StateStatus.SetTransition(EGameFeaturePluginState::Installed);
-			}
-		}
-	}
+	static_assert(UE::GameFeatures::NextStateTowardsTerminal[EGameFeaturePluginState::StatusKnown] == EGameFeaturePluginState::Uninstalling);
+	static_assert(UE::GameFeatures::NextStateTowardsActive[EGameFeaturePluginState::StatusKnown] == EGameFeaturePluginState::Downloading);
 };
 
 struct FGameFeaturePluginState_ErrorManagingData : public FErrorGameFeaturePluginState
 {
 	using FErrorGameFeaturePluginState::FErrorGameFeaturePluginState;
-
-	virtual void UpdateState(FGameFeaturePluginStateStatus& StateStatus) override
-	{
-		if (StateProperties.Destination < EGameFeaturePluginState::ErrorManagingData)
-		{			
-			StateStatus.SetTransition(EGameFeaturePluginState::Releasing);
-		}
-		else if (StateProperties.Destination > EGameFeaturePluginState::ErrorManagingData)
-		{
-			StateStatus.SetTransition(EGameFeaturePluginState::Downloading);
-		}
-	}
+	static_assert(UE::GameFeatures::NextStateTowardsTerminal[EGameFeaturePluginState::ErrorManagingData] == EGameFeaturePluginState::Releasing);
+	static_assert(UE::GameFeatures::NextStateTowardsActive[EGameFeaturePluginState::ErrorManagingData] == EGameFeaturePluginState::Downloading);
 };
 
 struct FGameFeaturePluginState_ErrorUninstalling : public FErrorGameFeaturePluginState
 {
 	using FErrorGameFeaturePluginState::FErrorGameFeaturePluginState;
 
-	virtual void UpdateState(FGameFeaturePluginStateStatus& StateStatus) override
-	{
-		if (StateProperties.Destination < EGameFeaturePluginState::ErrorUninstalling)
-		{
-			if (ShouldVisitUninstallStateBeforeTerminal())
-			{
-				StateStatus.SetTransition(EGameFeaturePluginState::Uninstalling);
-			}
-			else
-			{
-				StateStatus.SetTransition(EGameFeaturePluginState::Terminal);
-			}
-		}
-		else if (StateProperties.Destination > EGameFeaturePluginState::ErrorUninstalling)
-		{
-			StateStatus.SetTransition(EGameFeaturePluginState::StatusKnown);
-		}
-	}
+	static_assert(UE::GameFeatures::NextStateTowardsTerminal[EGameFeaturePluginState::ErrorUninstalling] == EGameFeaturePluginState::Uninstalling);
+	static_assert(UE::GameFeatures::NextStateTowardsActive[EGameFeaturePluginState::ErrorUninstalling] == EGameFeaturePluginState::StatusKnown);
 };
 
 //Base class for our states that want to request to release data from InstallBundleManager
@@ -1448,7 +1601,8 @@ struct FBaseDataReleaseGameFeaturePluginState : public FGameFeaturePluginState
 		{
 			case (EGameFeaturePluginProtocol::InstallBundle):
 			{
-				return true;
+				//InstallBundleProtocols have a MetaData that controlls if they uninstall currently
+				return StateProperties.ProtocolOptions.GetSubtype<FInstallBundlePluginProtocolOptions>().bUninstallBeforeTerminate;
 			}
 
 			default:
@@ -1485,17 +1639,8 @@ struct FGameFeaturePluginState_Uninstalled : public FDestinationGameFeaturePlugi
 {
 	using FDestinationGameFeaturePluginState::FDestinationGameFeaturePluginState;
 
-	virtual void UpdateState(FGameFeaturePluginStateStatus& StateStatus) override
-	{
-		if (StateProperties.Destination < EGameFeaturePluginState::Uninstalled)
-		{
-			StateStatus.SetTransition(EGameFeaturePluginState::Terminal);
-		}
-		else if (StateProperties.Destination > EGameFeaturePluginState::Uninstalled)
-		{
-			StateStatus.SetTransition(EGameFeaturePluginState::CheckingStatus);
-		}
-	}
+	static_assert(UE::GameFeatures::NextStateTowardsTerminal[EGameFeaturePluginState::Uninstalled] == EGameFeaturePluginState::Terminal);
+	static_assert(UE::GameFeatures::NextStateTowardsActive[EGameFeaturePluginState::Uninstalled] == EGameFeaturePluginState::CheckingStatus);
 };
 
 struct FGameFeaturePluginState_Uninstalling : public FBaseDataReleaseGameFeaturePluginState
@@ -1887,7 +2032,8 @@ struct FGameFeaturePluginState_Downloading : public FGameFeaturePluginState
 	{
 		Cleanup();
 
-		if (!ensure(ShouldDownloadContent()))
+		if ((StateProperties.GetPluginProtocol() == EGameFeaturePluginProtocol::File) 
+		|| !ensure(ShouldDownloadContent()))
 		{
 			bPluginDownloaded = true;
 			UpdateProgress(1.0f);
@@ -2079,69 +2225,33 @@ struct FGameFeaturePluginState_Installed : public FDestinationGameFeaturePluginS
 {
 	using FDestinationGameFeaturePluginState::FDestinationGameFeaturePluginState;
 
-	virtual void UpdateState(FGameFeaturePluginStateStatus& StateStatus) override
-	{
-		if (StateProperties.Destination > EGameFeaturePluginState::Installed)
-		{
-			StateStatus.SetTransition(EGameFeaturePluginState::Mounting);
-		}
-		else if (StateProperties.Destination < EGameFeaturePluginState::Installed)
-		{
-			StateStatus.SetTransition(EGameFeaturePluginState::Releasing);
-		}
-	}
+	static_assert(UE::GameFeatures::NextStateTowardsActive[EGameFeaturePluginState::Installed] == EGameFeaturePluginState::Mounting);
+	static_assert(UE::GameFeatures::NextStateTowardsTerminal[EGameFeaturePluginState::Installed] == EGameFeaturePluginState::Releasing);
 };
 
 struct FGameFeaturePluginState_ErrorMounting : public FErrorGameFeaturePluginState
 {
 	using FErrorGameFeaturePluginState::FErrorGameFeaturePluginState;
 
-	virtual void UpdateState(FGameFeaturePluginStateStatus& StateStatus) override
-	{
-		if (StateProperties.Destination < EGameFeaturePluginState::ErrorMounting)
-		{
-			StateStatus.SetTransition(EGameFeaturePluginState::Unmounting);
-		}
-		else if (StateProperties.Destination > EGameFeaturePluginState::ErrorMounting)
-		{
-			StateStatus.SetTransition(EGameFeaturePluginState::Mounting);
-		}
-	}
+	static_assert(UE::GameFeatures::NextStateTowardsTerminal[EGameFeaturePluginState::ErrorMounting] == EGameFeaturePluginState::Unmounting);
+	static_assert(UE::GameFeatures::NextStateTowardsActive[EGameFeaturePluginState::ErrorMounting] == EGameFeaturePluginState::Mounting);
 };
 
 struct FGameFeaturePluginState_ErrorWaitingForDependencies : public FErrorGameFeaturePluginState
 {
 	using FErrorGameFeaturePluginState::FErrorGameFeaturePluginState;
 
-	virtual void UpdateState(FGameFeaturePluginStateStatus& StateStatus) override
-	{
-		if (StateProperties.Destination < EGameFeaturePluginState::ErrorWaitingForDependencies)
-		{
-			// There is no cleaup state equivalent to EGameFeaturePluginState::WaitingForDependencies so just go back to unmounting
-			StateStatus.SetTransition(EGameFeaturePluginState::Unmounting);
-		}
-		else if (StateProperties.Destination > EGameFeaturePluginState::ErrorWaitingForDependencies)
-		{
-			StateStatus.SetTransition(EGameFeaturePluginState::WaitingForDependencies);
-		}
-	}
+	static_assert(UE::GameFeatures::NextStateTowardsTerminal[EGameFeaturePluginState::ErrorWaitingForDependencies] == EGameFeaturePluginState::Unmounting);
+	static_assert(UE::GameFeatures::NextStateTowardsActive[EGameFeaturePluginState::ErrorWaitingForDependencies] == EGameFeaturePluginState::WaitingForDependencies);
+
 };
 
 struct FGameFeaturePluginState_ErrorRegistering : public FErrorGameFeaturePluginState
 {
 	using FErrorGameFeaturePluginState::FErrorGameFeaturePluginState;
 
-	virtual void UpdateState(FGameFeaturePluginStateStatus& StateStatus) override
-	{
-		if (StateProperties.Destination < EGameFeaturePluginState::ErrorRegistering)
-		{
-			StateStatus.SetTransition(EGameFeaturePluginState::Unregistering);
-		}
-		else if (StateProperties.Destination > EGameFeaturePluginState::ErrorRegistering)
-		{
-			StateStatus.SetTransition(EGameFeaturePluginState::Registering);
-		}
-	}
+	static_assert(UE::GameFeatures::NextStateTowardsTerminal[EGameFeaturePluginState::ErrorRegistering] == EGameFeaturePluginState::Unregistering);
+	static_assert(UE::GameFeatures::NextStateTowardsActive[EGameFeaturePluginState::ErrorRegistering] == EGameFeaturePluginState::Registering);
 };
 
 struct FGameFeaturePluginState_Unmounting : public FGameFeaturePluginState
@@ -2286,7 +2396,7 @@ struct FGameFeaturePluginState_Unmounting : public FGameFeaturePluginState
 			bCheckedRealtimeMode = true;
 			if (UE::GameFeatures::RealtimeMode)
 			{
-				UE::GameFeatures::RealtimeMode->AddUpdateRequest(StateProperties.OnRequestUpdateStateMachine);
+				UE::GameFeatures::RealtimeMode->AddUpdateRequest(StateProperties);
 				return;
 			}
 		}
@@ -2308,6 +2418,8 @@ struct FGameFeaturePluginState_Unmounting : public FGameFeaturePluginState
 			return;
 		}
 
+		static_assert(UE::GameFeatures::NextStateTowardsTerminal[EGameFeaturePluginState::Unmounting] == EGameFeaturePluginState::Installed);
+		static_assert(UE::GameFeatures::NextStateTowardsActive[EGameFeaturePluginState::Unmounting] == EGameFeaturePluginState::Installed);
 		StateStatus.SetTransition(EGameFeaturePluginState::Installed);
 	}
 
@@ -2759,7 +2871,7 @@ struct FGameFeaturePluginState_Mounting : public FGameFeaturePluginState
 			bCheckedRealtimeMode = true;
 			if (UE::GameFeatures::RealtimeMode)
 			{
-				UE::GameFeatures::RealtimeMode->AddUpdateRequest(StateProperties.OnRequestUpdateStateMachine);
+				UE::GameFeatures::RealtimeMode->AddUpdateRequest(StateProperties);
 				return;
 			}
 		}
@@ -2795,6 +2907,8 @@ struct FGameFeaturePluginState_Mounting : public FGameFeaturePluginState
 		}
 		else
 		{
+			static_assert(UE::GameFeatures::NextStateTowardsTerminal[EGameFeaturePluginState::Mounting] == EGameFeaturePluginState::WaitingForDependencies);
+			static_assert(UE::GameFeatures::NextStateTowardsActive[EGameFeaturePluginState::Mounting] == EGameFeaturePluginState::WaitingForDependencies);
 			StateStatus.SetTransition(EGameFeaturePluginState::WaitingForDependencies);
 		}
 	}
@@ -2874,6 +2988,8 @@ struct FGameFeaturePluginState_AssetDependencyStreamOut : public FGameFeaturePlu
 
 	virtual void UpdateState(FGameFeaturePluginStateStatus& StateStatus) override
 	{
+		static_assert(UE::GameFeatures::NextStateTowardsActive[EGameFeaturePluginState::AssetDependencyStreamOut] == EGameFeaturePluginState::Unmounting);
+		static_assert(UE::GameFeatures::NextStateTowardsTerminal[EGameFeaturePluginState::AssetDependencyStreamOut] == EGameFeaturePluginState::Unmounting);
 		StateStatus.SetTransition(EGameFeaturePluginState::Unmounting);
 	}
 };
@@ -2882,17 +2998,8 @@ struct FGameFeaturePluginState_ErrorAssetDependencyStreaming : public FErrorGame
 {
 	using FErrorGameFeaturePluginState::FErrorGameFeaturePluginState;
 
-	virtual void UpdateState(FGameFeaturePluginStateStatus& StateStatus) override
-	{
-		if (StateProperties.Destination < EGameFeaturePluginState::ErrorAssetDependencyStreaming)
-		{
-			StateStatus.SetTransition(EGameFeaturePluginState::AssetDependencyStreamOut);
-		}
-		else if (StateProperties.Destination > EGameFeaturePluginState::ErrorAssetDependencyStreaming)
-		{
-			StateStatus.SetTransition(EGameFeaturePluginState::AssetDependencyStreaming);
-		}
-	}
+	static_assert(UE::GameFeatures::NextStateTowardsTerminal[EGameFeaturePluginState::ErrorAssetDependencyStreaming] == EGameFeaturePluginState::AssetDependencyStreamOut);
+	static_assert(UE::GameFeatures::NextStateTowardsActive[EGameFeaturePluginState::ErrorAssetDependencyStreaming] == EGameFeaturePluginState::AssetDependencyStreaming);
 };
 
 struct FGameFeaturePluginState_AssetDependencyStreaming : public FGameFeaturePluginState
@@ -3060,7 +3167,7 @@ struct FGameFeaturePluginState_AssetDependencyStreaming : public FGameFeaturePlu
 						&UE::IoStore::FOnDemandInstallProgress::Combine);
 					const float OverallProgress = SumProgress.GetRelativeProgress();
 
-					StateProperties.OnFeatureStateProgressUpdate.ExecuteIfBound(OverallProgress);
+					StateProperties.StateProgress = OverallProgress;
 				}));
 		}
 	}
@@ -3078,6 +3185,8 @@ struct FGameFeaturePluginState_AssetDependencyStreaming : public FGameFeaturePlu
 			return;
 		}
 
+		static_assert(UE::GameFeatures::NextStateTowardsActive[EGameFeaturePluginState::AssetDependencyStreaming] == EGameFeaturePluginState::Registering);
+		static_assert(UE::GameFeatures::NextStateTowardsTerminal[EGameFeaturePluginState::AssetDependencyStreaming] == EGameFeaturePluginState::Registering);
 		StateStatus.SetTransition(EGameFeaturePluginState::Registering);
 	}
 
@@ -3119,6 +3228,8 @@ struct FGameFeaturePluginState_Unregistering : public FGameFeaturePluginState
 	{
 		if (bHasUnloaded)
 		{
+			static_assert(UE::GameFeatures::NextStateTowardsActive[EGameFeaturePluginState::Unregistering] == EGameFeaturePluginState::AssetDependencyStreamOut);
+			static_assert(UE::GameFeatures::NextStateTowardsTerminal[EGameFeaturePluginState::Unregistering] == EGameFeaturePluginState::AssetDependencyStreamOut);
 			StateStatus.SetTransition(EGameFeaturePluginState::AssetDependencyStreamOut);
 			return;
 		}
@@ -3161,7 +3272,7 @@ struct FGameFeaturePluginState_Unregistering : public FGameFeaturePluginState
 		else
 		{
 			bRequestedUnloadPluginAssets = true;
-			UE::GameFeatures::ScheduleUnloadPluginAssets(StateProperties.PluginName, StateProperties.OnRequestUpdateStateMachine);
+			UE::GameFeatures::ScheduleUnloadPluginAssets(StateProperties.PluginName, StateProperties);
 		}
 #else
 		bHasUnloaded = true;
@@ -3449,7 +3560,7 @@ struct FGameFeaturePluginState_Registering : public FGameFeaturePluginState
 			bCheckedRealtimeMode = true;
 			if (UE::GameFeatures::RealtimeMode)
 			{
-				UE::GameFeatures::RealtimeMode->AddUpdateRequest(StateProperties.OnRequestUpdateStateMachine);
+				UE::GameFeatures::RealtimeMode->AddUpdateRequest(StateProperties);
 				return;
 			}
 		}
@@ -3474,6 +3585,8 @@ struct FGameFeaturePluginState_Registering : public FGameFeaturePluginState
 		{
 			check(LoadGFDState == ELoadGFDState::Success);
 
+			static_assert(UE::GameFeatures::NextStateTowardsActive[EGameFeaturePluginState::Registering] == EGameFeaturePluginState::Registered);
+			static_assert(UE::GameFeatures::NextStateTowardsTerminal[EGameFeaturePluginState::Registering] == EGameFeaturePluginState::Registered);
 			StateStatus.SetTransition(EGameFeaturePluginState::Registered);
 
 			check(StateProperties.AddedPrimaryAssetTypes.Num() == 0);
@@ -3501,34 +3614,16 @@ struct FGameFeaturePluginState_Registered : public FDestinationGameFeaturePlugin
 {
 	using FDestinationGameFeaturePluginState::FDestinationGameFeaturePluginState;
 
-	virtual void UpdateState(FGameFeaturePluginStateStatus& StateStatus) override
-	{
-		if (StateProperties.Destination > EGameFeaturePluginState::Registered)
-		{
-			StateStatus.SetTransition(EGameFeaturePluginState::Loading);
-		}
-		else if (StateProperties.Destination < EGameFeaturePluginState::Registered)
-		{
-			StateStatus.SetTransition( EGameFeaturePluginState::Unregistering);
-		}
-	}
+	static_assert(UE::GameFeatures::NextStateTowardsActive[EGameFeaturePluginState::Registered] == EGameFeaturePluginState::Loading);
+	static_assert(UE::GameFeatures::NextStateTowardsTerminal[EGameFeaturePluginState::Registered] == EGameFeaturePluginState::Unregistering);
 };
 
 struct FGameFeaturePluginState_ErrorLoading : public FErrorGameFeaturePluginState
 {
 	using FErrorGameFeaturePluginState::FErrorGameFeaturePluginState;
 
-	virtual void UpdateState(FGameFeaturePluginStateStatus& StateStatus) override
-	{
-		if (StateProperties.Destination < EGameFeaturePluginState::ErrorLoading)
-		{
-			StateStatus.SetTransition(EGameFeaturePluginState::Unloading);
-		}
-		else if (StateProperties.Destination > EGameFeaturePluginState::ErrorLoading)
-		{
-			StateStatus.SetTransition(EGameFeaturePluginState::Loading);
-		}
-	}
+	static_assert(UE::GameFeatures::NextStateTowardsTerminal[EGameFeaturePluginState::ErrorLoading] == EGameFeaturePluginState::Unloading);
+	static_assert(UE::GameFeatures::NextStateTowardsActive[EGameFeaturePluginState::ErrorLoading] == EGameFeaturePluginState::Loading);
 };
 
 struct FGameFeaturePluginState_Unloading : public FGameFeaturePluginState
@@ -3548,6 +3643,8 @@ struct FGameFeaturePluginState_Unloading : public FGameFeaturePluginState
 		UnloadGameFeatureBundles(StateProperties.GameFeatureData);
 		UGameFeaturesSubsystem::Get().OnGameFeatureUnloading(StateProperties.GameFeatureData, StateProperties.PluginIdentifier);
 
+		static_assert(UE::GameFeatures::NextStateTowardsActive[EGameFeaturePluginState::Unloading] == EGameFeaturePluginState::Registered);
+		static_assert(UE::GameFeatures::NextStateTowardsTerminal[EGameFeaturePluginState::Unloading] == EGameFeaturePluginState::Registered);
 		StateStatus.SetTransition(EGameFeaturePluginState::Registered);
 	}
 
@@ -3630,6 +3727,8 @@ struct FGameFeaturePluginState_Loading : public FGameFeaturePluginState
 
 		UGameFeaturesSubsystem::Get().OnGameFeatureLoading(StateProperties.GameFeatureData, StateProperties.PluginIdentifier);
 
+		static_assert(UE::GameFeatures::NextStateTowardsTerminal[EGameFeaturePluginState::Loading] == EGameFeaturePluginState::Loaded);
+		static_assert(UE::GameFeatures::NextStateTowardsActive[EGameFeaturePluginState::Loading] == EGameFeaturePluginState::Loaded);
 		StateStatus.SetTransition(EGameFeaturePluginState::Loaded);
 	}
 
@@ -3694,34 +3793,16 @@ struct FGameFeaturePluginState_Loaded : public FDestinationGameFeaturePluginStat
 {
 	using FDestinationGameFeaturePluginState::FDestinationGameFeaturePluginState;
 
-	virtual void UpdateState(FGameFeaturePluginStateStatus& StateStatus) override
-	{
-		if (StateProperties.Destination > EGameFeaturePluginState::Loaded)
-		{
-			StateStatus.SetTransition(EGameFeaturePluginState::ActivatingDependencies);
-		}
-		else if (StateProperties.Destination < EGameFeaturePluginState::Loaded)
-		{
-			StateStatus.SetTransition(EGameFeaturePluginState::Unloading);
-		}
-	}
+	static_assert(UE::GameFeatures::NextStateTowardsActive[EGameFeaturePluginState::Loaded] == EGameFeaturePluginState::ActivatingDependencies);
+	static_assert(UE::GameFeatures::NextStateTowardsTerminal[EGameFeaturePluginState::Loaded] == EGameFeaturePluginState::Unloading);
 };
 
 struct FGameFeaturePluginState_ErrorDeactivatingDependencies : public FErrorGameFeaturePluginState
 {
 	using FErrorGameFeaturePluginState::FErrorGameFeaturePluginState;
 
-	virtual void UpdateState(FGameFeaturePluginStateStatus& StateStatus) override
-	{
-		if (StateProperties.Destination < EGameFeaturePluginState::ErrorDeactivatingDependencies)
-		{
-			StateStatus.SetTransition(EGameFeaturePluginState::DeactivatingDependencies);
-		}
-		else if (StateProperties.Destination > EGameFeaturePluginState::ErrorDeactivatingDependencies)
-		{
-			StateStatus.SetTransition(EGameFeaturePluginState::DeactivatingDependencies);
-		}
-	}
+	static_assert(UE::GameFeatures::NextStateTowardsTerminal[EGameFeaturePluginState::ErrorDeactivatingDependencies] == EGameFeaturePluginState::DeactivatingDependencies);
+	static_assert(UE::GameFeatures::NextStateTowardsActive[EGameFeaturePluginState::ErrorDeactivatingDependencies] == EGameFeaturePluginState::DeactivatingDependencies);
 };
 
 struct FDeactivatingDependenciesTransitionPolicy
@@ -3815,6 +3896,8 @@ struct FGameFeaturePluginState_Deactivating : public FGameFeaturePluginState
 		if (bHasUnloaded)
 		{
 			check(NumExpectedPausers == NumObservedPausers);
+			static_assert(UE::GameFeatures::NextStateTowardsTerminal[EGameFeaturePluginState::Deactivating] == EGameFeaturePluginState::Loaded);
+			static_assert(UE::GameFeatures::NextStateTowardsActive[EGameFeaturePluginState::Deactivating] == EGameFeaturePluginState::Loaded);
 			StateStatus.SetTransition(EGameFeaturePluginState::Loaded);
 			return;
 		}
@@ -3871,18 +3954,8 @@ struct FGameFeaturePluginState_ErrorActivatingDependencies : public FErrorGameFe
 {
 	using FErrorGameFeaturePluginState::FErrorGameFeaturePluginState;
 
-	virtual void UpdateState(FGameFeaturePluginStateStatus& StateStatus) override
-	{
-		if (StateProperties.Destination < EGameFeaturePluginState::ErrorActivatingDependencies)
-		{
-			// There is no cleaup state equivalent to EGameFeaturePluginState::ErrorActivatingDependencies so just go back to Unloading
-			StateStatus.SetTransition(EGameFeaturePluginState::Unloading);
-		}
-		else if (StateProperties.Destination > EGameFeaturePluginState::ErrorActivatingDependencies)
-		{
-			StateStatus.SetTransition(EGameFeaturePluginState::ActivatingDependencies);
-		}
-	}
+	static_assert(UE::GameFeatures::NextStateTowardsTerminal[EGameFeaturePluginState::ErrorActivatingDependencies] == EGameFeaturePluginState::Unloading);
+	static_assert(UE::GameFeatures::NextStateTowardsActive[EGameFeaturePluginState::ErrorActivatingDependencies] == EGameFeaturePluginState::ActivatingDependencies);
 };
 
 struct FActivatingDependenciesTransitionPolicy
@@ -3974,6 +4047,8 @@ struct FGameFeaturePluginState_Activating : public FGameFeaturePluginState
 			UGameFeaturesSubsystem::Get().OnGameFeatureActivating(StateProperties.GameFeatureData, StateProperties.PluginName, Context, StateProperties.PluginIdentifier);
 		}
 
+		static_assert(UE::GameFeatures::NextStateTowardsTerminal[EGameFeaturePluginState::Activating] == EGameFeaturePluginState::Active);
+		static_assert(UE::GameFeatures::NextStateTowardsActive[EGameFeaturePluginState::Activating] == EGameFeaturePluginState::Active);
 		StateStatus.SetTransition(EGameFeaturePluginState::Active);
 	}
 
@@ -4012,13 +4087,8 @@ struct FGameFeaturePluginState_Active : public FDestinationGameFeaturePluginStat
 		}
 	}
 
-	virtual void UpdateState(FGameFeaturePluginStateStatus& StateStatus) override
-	{
-		if (StateProperties.Destination < EGameFeaturePluginState::Active)
-		{
-			StateStatus.SetTransition(EGameFeaturePluginState::DeactivatingDependencies);
-		}
-	}
+	static_assert(UE::GameFeatures::NextStateTowardsTerminal[EGameFeaturePluginState::Active] == EGameFeaturePluginState::DeactivatingDependencies);
+	static_assert(UE::GameFeatures::NextStateTowardsActive[EGameFeaturePluginState::Active] == EGameFeaturePluginState::Active);
 };
 
 /*
@@ -4083,7 +4153,6 @@ namespace UE::GameFeatures
 
 UGameFeaturePluginStateMachine::UGameFeaturePluginStateMachine(const FObjectInitializer& ObjectInitializer)
 	: Super(ObjectInitializer)
-	, CurrentStateInfo(EGameFeaturePluginState::Uninitialized)
 	, bInUpdateStateMachine(false)
 	, bRegisteredAsTransitioningGFPSM(false)
 {
@@ -4099,47 +4168,282 @@ FDestinationGameFeaturePluginState* AsDestinationState(FGameFeaturePluginState*
 	return nullptr;
 }
 
-void UGameFeaturesSubsystem::InitStateMachine(UGameFeaturePluginStateMachine* This, FGameFeaturePluginIdentifier InPluginIdentifier, const FGameFeatureProtocolOptions& InProtocolOptions) 
+UE::GameFeatures::FScopedPauseStateMachineExecution::FScopedPauseStateMachineExecution(UE::GameFeatures::FStateMachineExecutor* InExecutor)
+: Executor(InExecutor)
 {
-	LLM_SCOPE_BYTAG(GFP);
+	if (Executor)
+	{
+		bOldPaused = Executor->PauseExecution();	
+	}
+}
 
-	check(GetCurrentState(This) == EGameFeaturePluginState::Uninitialized);
-	FString PluginName{InPluginIdentifier.GetPluginName()};
-	This->CurrentStateInfo.State = EGameFeaturePluginState::UnknownStatus;
-	This->StateProperties = FGameFeaturePluginStateMachineProperties(
-		MoveTemp(InPluginIdentifier),
-		FGameFeaturePluginStateRange(This->CurrentStateInfo.State),
-		FGameFeaturePluginRequestUpdateStateMachine::CreateWeakLambda(this, [this, PluginName, WeakThis=MakeWeakObjectPtr(This)]()
-		{
-			if (UGameFeaturePluginStateMachine* This = WeakThis.Get())
+UE::GameFeatures::FScopedPauseStateMachineExecution::~FScopedPauseStateMachineExecution()
+{
+	if (Executor)
+	{
+		Executor->TryResumeExecution(bOldPaused);
+	}
+}
+
+UE::GameFeatures::FStateMachineExecutor::FStateMachineExecutor(UGameFeaturesSubsystem* InOwningSubsystem)
+	: OwningSubsystem(InOwningSubsystem)
+{
+}
+UE::GameFeatures::FStateMachineExecutor::~FStateMachineExecutor()
+{
+}
+
+bool UE::GameFeatures::FStateMachineExecutor::PauseExecution()
+{
+	bool OldValue = bPauseImmediateExecution;
+	bPauseImmediateExecution = true;
+	return OldValue;
+}
+
+void UE::GameFeatures::FStateMachineExecutor::TryResumeExecution(bool bStayPaused)
+{
+	bPauseImmediateExecution = bStayPaused;
+	if (!bPauseImmediateExecution)
+	{
+		TRACE_CPUPROFILER_EVENT_SCOPE_STR("FStateMachineExecutor::TryResumeExecution")
+		UpdateStateMachines();
+	}
+}
+
+bool UE::GameFeatures::FStateMachineExecutor::UpdateStateMachines(int32 MaxIterations)
+{
+	UAssetManager& LocalAssetManager = UAssetManager::Get();
+	LocalAssetManager.PushBulkScanning();
+	ON_SCOPE_EXIT 
+	{
+		LocalAssetManager.PopBulkScanning();
+	};
+
+	if (UE::GameFeatures::ExecutionOptions.bGrouped)
+	{
+		// New execution style, group by state and only allow state machines to make one transition at a time, hoisting the repetition to an outer loop here 
+		// Pause execution of re-entrant calls but do not use TryResumeExecution which will try and execute anything queued immediately.
+		TGuardValue<bool> GuardPause(bPauseImmediateExecution, true);
+		// Separate queues of GFPs to update by current state
+		TArray<UGameFeaturePluginStateMachine*> Queues[EGameFeaturePluginState::MAX];
+		for (int32 Iteration = 0; Iteration < MaxIterations; ++Iteration)
+		{ 
+			if (UE::GameFeatures::ExecutionOptions.bShuffle)
 			{
-				UpdateStateMachine(This);
+				Algo::RandomShuffle(OrderedUpdateQueue);
 			}
-			else
+			else if (UE::GameFeatures::ExecutionOptions.bReverse)
 			{
-				UE_LOGFMT(LogGameFeatures, Warning, "State machine pointer for {Plugin} became unexpectedly null.", PluginName);
+				Algo::Reverse(OrderedUpdateQueue);
 			}
-		}),
-		FGameFeatureStateProgressUpdate::CreateWeakLambda(this, [this, PluginName, WeakThis=MakeWeakObjectPtr(This)](float Progress)
-		{
-			if (UGameFeaturePluginStateMachine* This = WeakThis.Get())
+
+			// On every iteration, gather newly queued updates per state so that batches are as large as possible
+			for (FStateMachineId Id : OrderedUpdateQueue)
 			{
-				UpdateCurrentStateProgress(This, Progress);
+				UGameFeaturePluginStateMachine* StateMachine = OwningSubsystem->GameFeaturePluginStateMachines.FindRef(Id);
+				if (ensureMsgf(StateMachine, TEXT("State machine %lld wasn't present in subsystem when trying to update it"), Id))
+				{
+					StateMachine->StateProperties.bQueuedForUpdate = false;
+					Queues[StateMachine->StateProperties.CurrentState].Add(StateMachine);
+				}
 			}
-			else
+			OrderedUpdateQueue.Reset();
+
+			// Find the earliest state with work to do
+			EGameFeaturePluginState State = EGameFeaturePluginState::UnknownStatus;
+			for (EGameFeaturePluginState It : TEnumRange<EGameFeaturePluginState>())
+			{
+				if (Queues[It].Num() > 0)
+				{
+					State = It;
+					break;
+				}
+			}
+
+			if (Queues[State].IsEmpty())
+			{
+				// No more work to do
+				break;
+			}
+
+			UpdateStateMachinesGrouped(State, MakeArrayView(Queues[State]));
+			Queues[State].Reset();
+		}
+	}
+	else 
+	{
+		if (UE::GameFeatures::ExecutionOptions.bShuffle)
+		{
+			Algo::RandomShuffle(OrderedUpdateQueue);
+		}
+		else if (UE::GameFeatures::ExecutionOptions.bReverse)
+		{
+			Algo::Reverse(OrderedUpdateQueue);
+		}
+
+		// Closer to original execution style, update all deferred state machines and allow them to make as many transitions as they want, and update other state machines immediately
+		TArray<UGameFeaturePluginStateMachine*> LocalUpdateQueue;
+		LocalUpdateQueue.Reserve(OrderedUpdateQueue.Num());
+		for (const FStateMachineId& Id : OrderedUpdateQueue)
+		{
+			UGameFeaturePluginStateMachine* StateMachine = OwningSubsystem->GameFeaturePluginStateMachines.FindRef(Id);
+			if (ensureMsgf(StateMachine, TEXT("State machine %lld wasn't present in subsystem when trying to update it"), Id))
 			{
-				UE_LOGFMT(LogGameFeatures, Warning, "State machine pointer for {Plugin} became unexpectedly null.", PluginName);
+				LocalUpdateQueue.Add(StateMachine);
+				StateMachine->StateProperties.bQueuedForUpdate = false;
 			}
-		}));
+		}
+		OrderedUpdateQueue.Reset();
+
+		for (UGameFeaturePluginStateMachine* StateMachine : LocalUpdateQueue)
+		{
+			OwningSubsystem->UpdateStateMachine(StateMachine, {});
+		}
+	}
+	return OrderedUpdateQueue.Num() > 0;
+}
+
+#if 1
+#define GFP_STATE_OP(STATE_NAME, TXT) \
+template<> \
+void UE::GameFeatures::FStateMachineExecutor::UpdateStateMachinesGrouped<EGameFeaturePluginState::STATE_NAME>(TConstArrayView<UGameFeaturePluginStateMachine*> StateMachines) \
+{ \
+	TRACE_CPUPROFILER_EVENT_SCOPE_STR("UpdateStateMachinesGrouped<" #STATE_NAME ">"); \
+	for (UGameFeaturePluginStateMachine* StateMachine : StateMachines) \
+	{ \
+		OwningSubsystem->UpdateStateMachine(StateMachine, { 1 }); \
+	} \
+}
+	GAME_FEATURE_PLUGIN_STATE_LIST(GFP_STATE_OP)
+#undef GFP_STATE_OP
+
+void UE::GameFeatures::FStateMachineExecutor::UpdateStateMachinesGrouped(EGameFeaturePluginState State, TConstArrayView<UGameFeaturePluginStateMachine*> StateMachines)
+{
+	switch (State)
+	{
+#define GFP_STATE_OP(STATE_NAME, TXT) \
+	case EGameFeaturePluginState::STATE_NAME: \
+		UpdateStateMachinesGrouped<EGameFeaturePluginState::STATE_NAME>(StateMachines);\
+		break;
+	GAME_FEATURE_PLUGIN_STATE_LIST(GFP_STATE_OP)
+#undef GFP_STATE_OP
+	}
+}
+#else
+UE_TRACE_EVENT_BEGIN(Cpu, GFP_UpdateStateMachinesGrouped, NoSync)
+	UE_TRACE_EVENT_FIELD(UE::Trace::WideString, State)
+UE_TRACE_EVENT_END()
+
+void UE::GameFeatures::FStateMachineExecutor::UpdateStateMachinesGrouped(EGameFeaturePluginState State, TConstArrayView<UGameFeaturePluginStateMachine*> StateMachines)
+{
+	UE_TRACE_LOG_SCOPED_T(Cpu, GFP_UpdateStateMachinesGrouped, CpuChannel)
+		<< GFP_UpdateStateMachinesGrouped.State(*UE::GameFeatures::ToString(State));
+
+	for (UGameFeaturePluginStateMachine* StateMachine : StateMachines)
+	{
+		OwningSubsystem->UpdateStateMachine(StateMachine, { 1 });
+	}
+}
+#endif
+
+void UE::GameFeatures::FStateMachineExecutor::QueueUpdateStateMachine(FGameFeaturePluginStateMachineProperties& InProperties)
+{
+	if (!InProperties.bQueuedForUpdate)
+	{ 
+		ensureMsgf(OwningSubsystem->GameFeaturePluginStateMachines.FindRef(InProperties.Id), 
+			TEXT("State machine %lld wasn't present in subsystem when trying to update it"), InProperties.Id);
+		OwningSubsystem->EnableTick();
+		OrderedUpdateQueue.Add(InProperties.Id);
+		InProperties.bQueuedForUpdate = true;
+	}
+}
+
+void UE::GameFeatures::FStateMachineExecutor::TryUpdateStateMachineImmediate(int64 Id)
+{
+	UGameFeaturePluginStateMachine* StateMachine = OwningSubsystem->GameFeaturePluginStateMachines.FindRef(Id);
+	if (ensureMsgf(StateMachine, TEXT("State machine %lld wasn't present in subsystem when trying to update it immediately"), Id))
+	{
+		if (StateMachine->StateProperties.bQueuedForUpdate)
+		{
+			return;
+		}
+		else if (bPauseImmediateExecution)
+		{
+			QueueUpdateStateMachine(StateMachine->StateProperties);
+		}
+		else
+		{
+			TRACE_CPUPROFILER_EVENT_SCOPE_STR("FStateMachineExecutor Exectuting State Machine Immediately");
+			OwningSubsystem->UpdateStateMachine(StateMachine);
+		}
+	}
+}
+
+void UE::GameFeatures::FStateMachineExecutor::TryUpdateStateMachineImmediate(FGameFeaturePluginStateMachineProperties& InProperties)
+{
+	if (InProperties.bQueuedForUpdate)
+	{
+		return;
+	}
+	else if (bPauseImmediateExecution)
+	{
+		QueueUpdateStateMachine(InProperties);
+	}
+	else 
+	{
+		UGameFeaturePluginStateMachine* StateMachine = OwningSubsystem->GameFeaturePluginStateMachines.FindRef(InProperties.Id);
+		if (ensureMsgf(StateMachine, TEXT("State machine %lld wasn't present in subsystem when trying to update it immediately"), InProperties.Id))
+		{
+			OwningSubsystem->UpdateStateMachine(StateMachine);
+		}
+	}
+}
+
+void UGameFeaturesSubsystem::InitStateMachine(
+	UGameFeaturePluginStateMachine* This,
+	int64 InId,
+	FGameFeaturePluginIdentifier InPluginIdentifier,
+	const FGameFeatureProtocolOptions& InProtocolOptions) 
+{
+	LLM_SCOPE_BYTAG(GFP);
+
+	check(GetCurrentState(This) == EGameFeaturePluginState::Uninitialized);
+	FString PluginName{InPluginIdentifier.GetPluginName()};
+
+	FGameFeaturePluginRequestUpdateStateMachine UpdateCallback;
+	if (!StateMachineExecutor)
+	{
+		UpdateCallback.BindWeakLambda(this, 
+			[this, PluginName, WeakThis=MakeWeakObjectPtr(This)]()
+			{
+				if (UGameFeaturePluginStateMachine* This = WeakThis.Get())
+				{
+					UpdateStateMachine(This);
+				}
+				else
+				{
+					UE_LOGFMT(LogGameFeatures, Warning, "State machine pointer for {Plugin} became unexpectedly null.", PluginName);
+				}
+			});
+	}
+
+	constexpr EGameFeaturePluginState InitialState  = EGameFeaturePluginState::UnknownStatus;
+	This->StateProperties = FGameFeaturePluginStateMachineProperties(
+		InId,
+		MoveTemp(InPluginIdentifier),
+		InitialState,
+		StateMachineExecutor.Get(),
+		UpdateCallback);
 
 	This->StateProperties.ProtocolOptions = InProtocolOptions;
+	This->Executor = StateMachineExecutor.Get();
 
 #define GAME_FEATURE_PLUGIN_STATE_MAKE_STATE(inEnum, inText) This->AllStates[EGameFeaturePluginState::inEnum] = MakeUnique<FGameFeaturePluginState_##inEnum>(EGameFeaturePluginState::inEnum, This->StateProperties);
 	GAME_FEATURE_PLUGIN_STATE_LIST(GAME_FEATURE_PLUGIN_STATE_MAKE_STATE)
 #undef GAME_FEATURE_PLUGIN_STATE_MAKE_STATE
 
 	CheckAddBatchingRequestForCurrentState(This);
-	This->AllStates[This->CurrentStateInfo.State]->BeginState();
+	This->AllStates[This->StateProperties.CurrentState]->BeginState();
 }
 
 bool UGameFeaturesSubsystem::SetDestination(UGameFeaturePluginStateMachine* This, FGameFeaturePluginStateRange InDestination, FGameFeatureStateTransitionComplete OnFeatureStateTransitionComplete, FDelegateHandle* OutCallbackHandle /*= nullptr*/)
@@ -4156,7 +4460,7 @@ bool UGameFeaturesSubsystem::SetDestination(UGameFeaturePluginStateMachine* This
 	{
 		// Invalid range
 	}
-	else if (This->CurrentStateInfo.State == EGameFeaturePluginState::Terminal && !InDestination.Contains(EGameFeaturePluginState::Terminal))
+	else if (This->StateProperties.CurrentState == EGameFeaturePluginState::Terminal && !InDestination.Contains(EGameFeaturePluginState::Terminal))
 	{
 		// Can't tranistion away from terminal state
 	}
@@ -4169,15 +4473,15 @@ bool UGameFeaturesSubsystem::SetDestination(UGameFeaturePluginStateMachine* This
 			OutCallbackHandle->Reset();
 		}
 
-		FDestinationGameFeaturePluginState* CurrState = This->AllStates[This->CurrentStateInfo.State]->AsDestinationState();
+		FDestinationGameFeaturePluginState* CurrState = This->AllStates[This->StateProperties.CurrentState]->AsDestinationState();
 
-		if (InDestination.Contains(This->CurrentStateInfo.State))
+		if (InDestination.Contains(This->StateProperties.CurrentState))
 		{
 			OnFeatureStateTransitionComplete.ExecuteIfBound(This, MakeValue());
 		}
 		else
 		{
-			if (This->CurrentStateInfo.State < InDestination)
+			if (This->StateProperties.CurrentState < InDestination)
 			{
 				FDestinationGameFeaturePluginState* MinDestState = This->AllStates[InDestination.MinState]->AsDestinationState();
 				FDelegateHandle CallbackHandle = MinDestState->OnDestinationStateReached.Add(MoveTemp(OnFeatureStateTransitionComplete));
@@ -4186,7 +4490,7 @@ bool UGameFeaturesSubsystem::SetDestination(UGameFeaturePluginStateMachine* This
 					*OutCallbackHandle = CallbackHandle;
 				}
 			}
-			else if (This->CurrentStateInfo.State > InDestination)
+			else if (This->StateProperties.CurrentState > InDestination)
 			{
 				FDestinationGameFeaturePluginState* MaxDestState = This->AllStates[InDestination.MaxState]->AsDestinationState();
 				FDelegateHandle CallbackHandle = MaxDestState->OnDestinationStateReached.Add(MoveTemp(OnFeatureStateTransitionComplete));
@@ -4197,7 +4501,7 @@ bool UGameFeaturesSubsystem::SetDestination(UGameFeaturePluginStateMachine* This
 			}
 
 			This->StateProperties.Destination = InDestination;
-			UpdateStateMachine(This);
+			This->StateProperties.UpdateStateMachineImmediate();
 			bDestinationChanged = true;
 		}
 
@@ -4213,12 +4517,13 @@ bool UGameFeaturesSubsystem::SetDestination(UGameFeaturePluginStateMachine* This
 			OutCallbackHandle->Reset();
 		}
 
-		if (This->CurrentStateInfo.State < This->StateProperties.Destination)
+		if (This->StateProperties.CurrentState < This->StateProperties.Destination)
 		{
 			This->StateProperties.Destination = *NewDestination;
 
-			if (InDestination.Contains(This->CurrentStateInfo.State))
+			if (InDestination.Contains(This->StateProperties.CurrentState))
 			{
+				// TODO: Consider deferring for executor 
 				OnFeatureStateTransitionComplete.ExecuteIfBound(This, MakeValue());
 			}
 			else
@@ -4232,12 +4537,13 @@ bool UGameFeaturesSubsystem::SetDestination(UGameFeaturePluginStateMachine* This
 				bDestinationChanged = true;
 			}
 		}
-		else if(This->CurrentStateInfo.State > This->StateProperties.Destination)
+		else if(This->StateProperties.CurrentState > This->StateProperties.Destination)
 		{
 			This->StateProperties.Destination = *NewDestination;
 
-			if (InDestination.Contains(This->CurrentStateInfo.State))
+			if (InDestination.Contains(This->StateProperties.CurrentState))
 			{
+				// TODO: Consider deferring for executor 
 				OnFeatureStateTransitionComplete.ExecuteIfBound(This, MakeValue());
 			}
 			else
@@ -4414,14 +4720,14 @@ FGameFeaturePluginStateRange UGameFeaturesSubsystem::GetDestination(const UGameF
 	return This->StateProperties.Destination;
 }
 
-const FGameFeaturePluginStateInfo& UGameFeaturesSubsystem::GetCurrentStateInfo(const UGameFeaturePluginStateMachine* This) const
+FGameFeaturePluginStateInfo UGameFeaturesSubsystem::GetCurrentStateInfo(const UGameFeaturePluginStateMachine* This) const
 {
-	return This->CurrentStateInfo;
+	return FGameFeaturePluginStateInfo{ This->StateProperties.CurrentState, This->StateProperties.StateProgress };
 }
 
 bool UGameFeaturesSubsystem::IsRunning(const UGameFeaturePluginStateMachine* This) const
 {
-	return !This->StateProperties.Destination.Contains(This->CurrentStateInfo.State);
+	return !This->StateProperties.Destination.Contains(This->StateProperties.CurrentState);
 }
 
 bool UGameFeaturesSubsystem::IsStatusKnown(const UGameFeaturePluginStateMachine* This) const
@@ -4562,7 +4868,7 @@ UE_TRACE_EVENT_BEGIN(Cpu, GFP_UpdateStateMachine_EndState, NoSync)
 	UE_TRACE_EVENT_FIELD(UE::Trace::WideString, State)
 UE_TRACE_EVENT_END()
 
-void UGameFeaturesSubsystem::UpdateStateMachine(UGameFeaturePluginStateMachine* This) 
+void UGameFeaturesSubsystem::UpdateStateMachine(UGameFeaturePluginStateMachine* This, TOptional<int32> InMaxTransitions) 
 {
 	LLM_SCOPE_BYTAG(GFP);
 
@@ -4609,7 +4915,7 @@ void UGameFeaturesSubsystem::UpdateStateMachine(UGameFeaturePluginStateMachine*
 
 	bool bKeepProcessing = false;
 	int32 NumTransitions = 0;
-	const int32 MaxTransitions = 10000;
+	const int32 MaxTransitions = InMaxTransitions.Get(10000);
 	do
 	{
 		bKeepProcessing = false;
@@ -4639,7 +4945,8 @@ void UGameFeaturesSubsystem::UpdateStateMachine(UGameFeaturePluginStateMachine*
 				This->AllStates[CurrentState]->EndState();
 				CheckAndCancelBatchingRequestForCurrentState(This);
 			}
-			This->CurrentStateInfo = FGameFeaturePluginStateInfo(StateStatus.TransitionToState);
+			This->StateProperties.CurrentState = StateStatus.TransitionToState;
+			This->StateProperties.StateProgress = 0.0f;
 			CurrentState = StateStatus.TransitionToState;
 			check(CurrentState != EGameFeaturePluginState::MAX);
 
@@ -4763,18 +5070,20 @@ void UGameFeaturesSubsystem::UpdateStateMachine(UGameFeaturePluginStateMachine*
 			}
 		}
 
-		if (NumTransitions++ > MaxTransitions)
+		if (++NumTransitions >= MaxTransitions)
 		{
-			UE_LOG(LogGameFeatures, Fatal, TEXT("Infinite loop in game feature state machine transitions. Current state %s. GameFeature: %s"), *UE::GameFeatures::ToString(CurrentState), *GetGameFeatureName(This));
+			UE_CLOG(!InMaxTransitions.IsSet(), LogGameFeatures, Fatal, TEXT("Infinite loop in game feature state machine transitions. Current state %s. GameFeature: %s"), *UE::GameFeatures::ToString(CurrentState), *GetGameFeatureName(This));
+			// If we aren't attempting to do all transitions at once, make sure we queue up another UpdateStateMachine if there was still work to do
+			// (e.g. we just transitioned to a new state which wants UpdateState to be called for the first time)
+			if (bKeepProcessing)
+			{
+				This->StateProperties.UpdateStateMachineDeferred(0.0f);
+			}
+			break;
 		}
 	} while (bKeepProcessing);
 }
 
-void UGameFeaturesSubsystem::UpdateCurrentStateProgress(UGameFeaturePluginStateMachine* This, float Progress) const
-{
-	This->CurrentStateInfo.Progress = Progress;
-}
-
 void UGameFeaturesSubsystem::RegisterAsTransitioningStateMachine(UGameFeaturePluginStateMachine* This) 
 {
 	if (This->bRegisteredAsTransitioningGFPSM)
@@ -4802,15 +5111,15 @@ void UGameFeaturesSubsystem::CheckAddBatchingRequestForCurrentState(UGameFeature
 	check(!This->StateProperties.BatchProcessingHandle.IsValid());
 
 	const bool bCanBatchProcess =
-		UE::GameFeatures::BatchProcessingHelperFunctors[This->CurrentStateInfo.State].ImplementsBatchProcess() &&
+		UE::GameFeatures::BatchProcessingHelperFunctors[This->StateProperties.CurrentState].ImplementsBatchProcess() &&
 		This->StateProperties.CanBatchProcess() &&
-		This->AllStates[This->CurrentStateInfo.State]->CanBatchProcess();
+		This->AllStates[This->StateProperties.CurrentState]->CanBatchProcess();
 	
 	if (bCanBatchProcess)
 	{
-		UE_LOG(LogGameFeatures, Verbose, TEXT("Game feature '%s' awaiting batch processing of state (%s)"), *GetGameFeatureName(This), *UE::GameFeatures::ToString(This->CurrentStateInfo.State));
+		UE_LOG(LogGameFeatures, Verbose, TEXT("Game feature '%s' awaiting batch processing of state (%s)"), *GetGameFeatureName(This), *UE::GameFeatures::ToString(This->StateProperties.CurrentState));
 		This->StateProperties.BatchProcessingHandle = AddBatchingRequest(
-			This->CurrentStateInfo.State,
+			This->StateProperties.CurrentState,
 			This->StateProperties.OnRequestUpdateStateMachine);
 	}
 }
@@ -4823,9 +5132,9 @@ void UGameFeaturesSubsystem::CheckAndCancelBatchingRequestForCurrentState(UGameF
 	// If we are currently awaiting batch processing, cancel and update the state machine.
 	if (This->StateProperties.IsWaitingForBatchProcessing())
 	{
-		UE_LOG(LogGameFeatures, Verbose, TEXT("Game feature '%s' cancelled batch processing of state (%s)"), *GetGameFeatureName(This), *UE::GameFeatures::ToString(This->CurrentStateInfo.State));
+		UE_LOG(LogGameFeatures, Verbose, TEXT("Game feature '%s' cancelled batch processing of state (%s)"), *GetGameFeatureName(This), *UE::GameFeatures::ToString(This->StateProperties.CurrentState));
 		CancelBatchingRequest(
-			This->CurrentStateInfo.State,
+			This->StateProperties.CurrentState,
 			This->StateProperties.BatchProcessingHandle);
 
 		This->StateProperties.BatchProcessingHandle.Reset();
@@ -4860,18 +5169,55 @@ void UGameFeaturesSubsystem::ExcludeFromBatchProcessing(UGameFeaturePluginStateM
 	}	
 }
 
+void UGameFeaturesSubsystem::TryUpdateStateMachineImmediate(int64 Id)
+{
+	if (StateMachineExecutor)
+	{
+		StateMachineExecutor->TryUpdateStateMachineImmediate(Id);
+	}
+	else if (UGameFeaturePluginStateMachine* StateMachine = GameFeaturePluginStateMachines.FindRef(Id))
+	{
+		UpdateStateMachine(StateMachine);
+	}
+}
+
 FGameFeaturePluginStateMachineProperties::FGameFeaturePluginStateMachineProperties(
+	int64 InId,
 	FGameFeaturePluginIdentifier InPluginIdentifier,
-	const FGameFeaturePluginStateRange& DesiredDestination,
-	const FGameFeaturePluginRequestUpdateStateMachine& RequestUpdateStateMachineDelegate,
-	const FGameFeatureStateProgressUpdate& FeatureStateProgressUpdateDelegate)
-	: PluginIdentifier(MoveTemp(InPluginIdentifier))
-	, Destination(DesiredDestination)
+	EGameFeaturePluginState InitialState,
+	UE::GameFeatures::FStateMachineExecutor* InExecutor,
+	const FGameFeaturePluginRequestUpdateStateMachine& RequestUpdateStateMachineDelegate)
+	: Id(InId)
+	, CurrentState(InitialState)
+	, PluginIdentifier(MoveTemp(InPluginIdentifier))
+	, Destination(FGameFeaturePluginStateRange(InitialState))
 	, OnRequestUpdateStateMachine(RequestUpdateStateMachineDelegate)
-	, OnFeatureStateProgressUpdate(FeatureStateProgressUpdateDelegate)
+	, Executor(InExecutor)
 {
 }
 
+void FGameFeaturePluginStateMachineProperties::UpdateStateMachineImmediate() 
+{
+	if (Executor)
+	{
+		Executor->TryUpdateStateMachineImmediate(*this);
+	}
+	else
+	{
+		OnRequestUpdateStateMachine.ExecuteIfBound();
+	}
+}
+
+bool FGameFeaturePluginStateMachineProperties::UpdateStateMachineDeferred(float Delay) 
+{
+	if (Executor)
+	{
+		Executor->QueueUpdateStateMachine(*this);
+		return true;
+	}
+	return false;
+}
+
 EGameFeaturePluginProtocol FGameFeaturePluginStateMachineProperties::GetPluginProtocol() const
 {
 	return PluginIdentifier.GetPluginProtocol();
@@ -5050,7 +5396,7 @@ bool FGameFeaturePluginStateMachineProperties::AllowAsyncLoading() const
 
 bool FGameFeaturePluginStateMachineProperties::CanBatchProcess() const
 {
-	return ProtocolOptions.bBatchProcess && UE::GameFeatures::CVarEnableBatchProcessing.GetValueOnGameThread();
+	return ProtocolOptions.bBatchProcess && !UE::GameFeatures::bUseNewExecutor && UE::GameFeatures::CVarEnableBatchProcessing.GetValueOnGameThread();
 }
 
 bool FGameFeaturePluginStateMachineProperties::IsWaitingForBatchProcessing() const
diff --git a/Engine/Plugins/Runtime/GameFeatures/Source/GameFeatures/Private/GameFeaturePluginStateMachine.h b/Engine/Plugins/Runtime/GameFeatures/Source/GameFeatures/Private/GameFeaturePluginStateMachine.h
--- a/Engine/Plugins/Runtime/GameFeatures/Source/GameFeatures/Private/GameFeaturePluginStateMachine.h
+++ b/Engine/Plugins/Runtime/GameFeatures/Source/GameFeatures/Private/GameFeaturePluginStateMachine.h
@@ -18,6 +18,8 @@ enum class EInstallBundleReleaseResult : uint32;
 namespace UE::GameFeatures
 {
 	extern TAutoConsoleVariable<bool> CVarAllowMissingOnDemandDependencies;
+	extern bool bUseNewExecutor;
+	extern bool bLoadBuiltin_DeferExecution;
 }
 
 /*
@@ -172,6 +174,76 @@ Transition states are expected to transition the machine to another state after
                              +--------------------+
 */
 
+namespace UE::GameFeatures
+{
+	class FStateMachineExecutor
+	{
+		// State machines are addressed by FGameFeaturePluginStateMachineProperties::Id
+		using FStateMachineId = int64;
+
+		// Lists of state machines requiring updates stored by the order in which a deferred update request was made
+		TArray<FStateMachineId> OrderedUpdateQueue;
+		UGameFeaturesSubsystem* OwningSubsystem = nullptr;
+
+		/** 
+		 * If this is true, state machine update requests through TryUpdateStateMachineImmediate will be queued as if
+		 * QueueUpdateStateMachine had been called instead 
+		 */
+		bool bPauseImmediateExecution = false;
+
+		/** Update a set of state machines which are expected to be in the same state */
+		void UpdateStateMachinesGrouped(EGameFeaturePluginState State, TConstArrayView<UGameFeaturePluginStateMachine*> StateMachines);
+		template<EGameFeaturePluginState State>
+		void UpdateStateMachinesGrouped(TConstArrayView<UGameFeaturePluginStateMachine*> StateMachines);
+
+	public:
+		FStateMachineExecutor(UGameFeaturesSubsystem* InOwningSubsystem);
+		~FStateMachineExecutor();
+		UE_NONCOPYABLE(FStateMachineExecutor);
+
+		/** Update all state machines which are not at their target state and are not waiting on async work to complete */
+		bool UpdateStateMachines(int32 MaxIterations = 10000);	
+
+		/** 
+		 * Pause immediate execution of state machine updates and return whether it was already paused.
+		 * The return value should be passed back to TryResumeExecution from a scoped object. 
+		 */
+		bool PauseExecution();
+
+		/** 
+		 * Try and resume immediately executing state machines.
+		 * If the argument is false, execution is resumed and any queued updates are performed.
+		 */
+		void TryResumeExecution(bool bInStayPaused);
+
+		/**
+		 * Store the given state machine as needing to be updated to transition to the next state 
+		 * TODO: Do we need a delay timer, or can that be replaced with proper events?
+		 */
+		void QueueUpdateStateMachine(FGameFeaturePluginStateMachineProperties& InProperties);
+
+		/**
+		 * If immediate execution is allowed, update the given state machine.
+		 * Otherwise queue it for deferred execution.
+		 */
+		void TryUpdateStateMachineImmediate(FGameFeaturePluginStateMachineProperties& InProperties);
+		void TryUpdateStateMachineImmediate(int64 Id);
+	};
+
+	struct FScopedPauseStateMachineExecution
+	{
+		FScopedPauseStateMachineExecution(FStateMachineExecutor* InExecutor);
+		~FScopedPauseStateMachineExecution();
+		UE_NONCOPYABLE(FScopedPauseStateMachineExecution);
+	
+	private:
+		FStateMachineExecutor* Executor = nullptr;
+		bool bOldPaused = false;
+	};
+
+}
+
+/* Represents an inclusive range of states which it is acceptable for a GFP to be in */
 struct FGameFeaturePluginStateRange
 {
 	EGameFeaturePluginState MinState = EGameFeaturePluginState::Uninitialized;
@@ -189,16 +261,19 @@ struct FGameFeaturePluginStateRange
 
 	bool IsValid() const { return MinState <= MaxState; }
 
+	/* Returns whether a given single state is inside this range */
 	bool Contains(EGameFeaturePluginState InState) const
 	{
 		return InState >= MinState && InState <= MaxState;
 	}
 
+	/* Returns whether there are any states in common between this range and another range */
 	bool Overlaps(const FGameFeaturePluginStateRange& Other) const
 	{
 		return Other.MinState <= MaxState && Other.MaxState >= MinState;
 	}
 
+	/* If there are any states in common between this range and another range, return a new range representing those states. */
 	TOptional<FGameFeaturePluginStateRange> Intersect(const FGameFeaturePluginStateRange& Other) const
 	{
 		TOptional<FGameFeaturePluginStateRange> Intersection;
@@ -274,6 +349,21 @@ struct FGameFeaturePluginStateMachineProperties
 {
 	GENERATED_BODY()
 
+	/** 
+	 * Integer identifier for this plugin when the UGameFeaturePluginStateMachine was created.
+	 * Is not stable across termination and recreation of the state machine for this plugin. 
+	 */
+	int64 Id = 0;
+
+	/** Current state of the state machine. Should not be modified by states themselves, only by UpdateStateMachine or FStateMachineExecutor */
+	EGameFeaturePluginState CurrentState = EGameFeaturePluginState::Uninitialized;
+	
+	/** Whether this GFP is in the queue in FStateMachineExecutor */
+	bool bQueuedForUpdate = false;
+
+	/** The progress of the current state. Relevant only for transition states. */
+	float StateProgress = 0.0f;
+
 	/** 
 	* The Identifier used to find this Plugin. Parsed from the supplied PluginURL at creation.
 	* Every protocol will have its own style of identifier URL that will get parsed to generate this.
@@ -324,18 +414,13 @@ struct FGameFeaturePluginStateMachineProperties
 	DECLARE_MULTICAST_DELEGATE_OneParam(FOnTransitionCanceled, UGameFeaturePluginStateMachine* /*Machine*/);
 	FOnTransitionCanceled OnTransitionCanceled;
 
-	/** Delegate to request the state machine be updated. */
-	FGameFeaturePluginRequestUpdateStateMachine OnRequestUpdateStateMachine;
-
-	/** Delegate for when a feature state needs to update progress. */
-	FGameFeatureStateProgressUpdate OnFeatureStateProgressUpdate;
-
 	FGameFeaturePluginStateMachineProperties() = default;
 	FGameFeaturePluginStateMachineProperties(
+		int64 InId,
 		FGameFeaturePluginIdentifier InPluginIdentifier,
-		const FGameFeaturePluginStateRange& DesiredDestination,
-		const FGameFeaturePluginRequestUpdateStateMachine& RequestUpdateStateMachineDelegate,
-		const FGameFeatureStateProgressUpdate& FeatureStateProgressUpdateDelegate);
+		EGameFeaturePluginState InitialState,
+		UE::GameFeatures::FStateMachineExecutor* Executor,
+		const FGameFeaturePluginRequestUpdateStateMachine& RequestUpdateStateMachineDelegate);
 
 	EGameFeaturePluginProtocol GetPluginProtocol() const;
 
@@ -353,6 +438,22 @@ struct FGameFeaturePluginStateMachineProperties
 	bool CanBatchProcess() const;
 	bool IsWaitingForBatchProcessing() const;
 	bool WasBatchProcessed() const;
+
+	void UpdateStateMachineImmediate();
+	/** 
+	 * Queues an update for this state machine in the executor if it is being used and returns true.
+	 * Returns false if the executor is not being used 
+	 */
+	bool UpdateStateMachineDeferred(float Delay);
+
+private:
+	friend class UGameFeaturesSubsystem;
+
+	/** Delegate to request the state machine be updated. */
+	FGameFeaturePluginRequestUpdateStateMachine OnRequestUpdateStateMachine;
+	
+	/* Object responsible for scheduling state machine updates when UE::GameFeatures::bUseNewExecutor is true */
+	UE::GameFeatures::FStateMachineExecutor* Executor = nullptr;
 };
 
 /** Input and output information for a state's UpdateState */
@@ -459,9 +560,17 @@ struct FDestinationGameFeaturePluginState : public FGameFeaturePluginState
 {
 	using FGameFeaturePluginState::FGameFeaturePluginState;
 
+	/** Called when this state is no longer the active state */
+	virtual void EndState() override final;
+	/** Attempt to cancel any pending state transition. */
+	virtual void TryCancelState() override final;
+
 	/** Returns the type of state this is */
 	virtual EGameFeaturePluginStateType GetStateType() const override { return EGameFeaturePluginStateType::Destination; }
 
+	/** Process the state's logic to decide if there should be a state transition. */
+	virtual void UpdateState(FGameFeaturePluginStateStatus& StateStatus) override final;
+
 	DECLARE_MULTICAST_DELEGATE_TwoParams(FOnDestinationStateReached, UGameFeaturePluginStateMachine* /*Machine*/, const UE::GameFeatures::FResult& /*Result*/);
 	FOnDestinationStateReached OnDestinationStateReached;
 };
@@ -485,7 +594,11 @@ struct FGameFeaturePluginStateInfo
 	float Progress = 0.0f;
 
 	FGameFeaturePluginStateInfo() = default;
-	explicit FGameFeaturePluginStateInfo(EGameFeaturePluginState InState) : State(InState) {}
+	explicit FGameFeaturePluginStateInfo(EGameFeaturePluginState InState, float InProgress = 0.0f)
+		: State(InState)
+		, Progress(InProgress)
+	{
+	}
 };
 
 /** A state machine to manage transitioning a game feature plugin from just a URL into a fully loaded and active plugin, including registering its contents with other game systems */
@@ -499,8 +612,6 @@ public:
 
 	FDestinationGameFeaturePluginState* AsDestinationState(FGameFeaturePluginState* InState);
 
-	/** Information about the current state */
-	FGameFeaturePluginStateInfo CurrentStateInfo;
 
 	/** The common properties that can be accessed by the states of the state machine */
 	UPROPERTY(transient)
@@ -509,6 +620,9 @@ public:
 	/** All state machine state objects */
 	TUniquePtr<FGameFeaturePluginState> AllStates[EGameFeaturePluginState::MAX];
 
+	/* Object responsible for scheduling state machine updates when UE::GameFeatures::bUseNewExecutor is true */
+	UE::GameFeatures::FStateMachineExecutor* Executor = nullptr;
+
 	/** True when we are currently executing UpdateStateMachine, to avoid reentry */
 	bool bInUpdateStateMachine;
 
diff --git a/Engine/Plugins/Runtime/GameFeatures/Source/GameFeatures/Private/GameFeaturesSubsystem.cpp b/Engine/Plugins/Runtime/GameFeatures/Source/GameFeatures/Private/GameFeaturesSubsystem.cpp
--- a/Engine/Plugins/Runtime/GameFeatures/Source/GameFeatures/Private/GameFeaturesSubsystem.cpp
+++ b/Engine/Plugins/Runtime/GameFeatures/Source/GameFeatures/Private/GameFeaturesSubsystem.cpp
@@ -505,6 +505,15 @@ FGameFeatureProtocolOptions::FGameFeatureProtocolOptions(FNull InOptions)
 	SetSubtype<FNull>(InOptions);
 }
 
+UGameFeaturesSubsystem::UGameFeaturesSubsystem(const FObjectInitializer& ObjectInitializer)
+	: Super()
+{
+	if (UE::GameFeatures::bUseNewExecutor && !HasAnyFlags(RF_ClassDefaultObject))
+	{
+		StateMachineExecutor = MakePimpl<UE::GameFeatures::FStateMachineExecutor>(this);
+	}
+}
+
 void UGameFeaturesSubsystem::Initialize(FSubsystemCollectionBase& Collection)
 {
 	UE_LOG(LogGameFeatures, Log, TEXT("Initializing game features subsystem"));
@@ -699,7 +708,7 @@ void UGameFeaturesSubsystem::Deinitialize()
 
 	// Delete all state objects so that they don't try to run destructor logic during GC/shutdown in a way that could access 
 	// already-deleted objects (e.g. GEngine, this subsystem)
-	for (TPair<FString, TObjectPtr<UGameFeaturePluginStateMachine>>& Pair : GameFeaturePluginStateMachines)
+	for (TPair<int64, TObjectPtr<UGameFeaturePluginStateMachine>>& Pair : GameFeaturePluginStateMachines)
 	{
 		if (Pair.Value)
 		{
@@ -2834,6 +2843,8 @@ void UGameFeaturesSubsystem::LoadBuiltInGameFeaturePluginsInternal(TArray<TShare
 	}
 	else
 	{
+		// Conditionally defer execution if the new executor is being used and deferred execution is enabled 
+		UE::GameFeatures::FScopedPauseStateMachineExecution PauseExecution(UE::GameFeatures::bLoadBuiltin_DeferExecution ? StateMachineExecutor.Get() : nullptr);
 		for (const TSharedRef<IPlugin>& Plugin : Plugins)
 		{
 			FBuiltInPluginLoadTimeTrackerScope TrackerScope(PluginLoadTimeTracker, Plugin);
@@ -3788,7 +3799,7 @@ struct FGameFeaturePluginPredownloadContext : public FGameFeaturePluginPredownlo
 
 		TArray<FName> ReleaseList;
 		TSet<FName> KeepList;
-		for (const TPair<FString, TObjectPtr<UGameFeaturePluginStateMachine>>& Pair : GFPSubSys.GameFeaturePluginStateMachines)
+		for (const TPair<int64, TObjectPtr<UGameFeaturePluginStateMachine>>& Pair : GFPSubSys.GameFeaturePluginStateMachines)
 		{
 			UGameFeaturePluginStateMachine* Machine = Pair.Value;
 			if (Machine &&
@@ -3850,7 +3861,7 @@ UGameFeaturePluginStateMachine* UGameFeaturesSubsystem::FindGameFeaturePluginSta
 	const FStringView ShortUrl = PluginIdentifier.GetIdentifyingString();
 
 	TObjectPtr<UGameFeaturePluginStateMachine> const* ExistingStateMachine = 
-		GameFeaturePluginStateMachines.FindByHash(GetTypeHash(PluginIdentifier.GetIdentifyingString()), PluginIdentifier.GetIdentifyingString());
+		GameFeaturesByIdentifier.FindByHash(GetTypeHash(PluginIdentifier.GetIdentifyingString()), PluginIdentifier.GetIdentifyingString());
 	if (ExistingStateMachine)
 	{
 		EGameFeaturePluginProtocol ExpectedProtocol = GetPluginIdentifier(*ExistingStateMachine).GetPluginProtocol();
@@ -3872,7 +3883,7 @@ UGameFeaturePluginStateMachine* UGameFeaturesSubsystem::FindOrCreateGameFeatureP
 	TRACE_CPUPROFILER_EVENT_SCOPE(GFP_FindOrCreateStateMachine);
 	FGameFeaturePluginIdentifier PluginIdentifier(PluginURL);
 	TObjectPtr<UGameFeaturePluginStateMachine> const* ExistingStateMachine =
-		GameFeaturePluginStateMachines.FindByHash(GetTypeHash(PluginIdentifier.GetIdentifyingString()), PluginIdentifier.GetIdentifyingString());
+		GameFeaturesByIdentifier.FindByHash(GetTypeHash(PluginIdentifier.GetIdentifyingString()), PluginIdentifier.GetIdentifyingString());
 
 	if (bOutFoundExisting)
 	{
@@ -3901,8 +3912,10 @@ UGameFeaturePluginStateMachine* UGameFeaturesSubsystem::FindOrCreateGameFeatureP
 #endif
 
 	UGameFeaturePluginStateMachine* NewStateMachine = NewObject<UGameFeaturePluginStateMachine>(this);
-	GameFeaturePluginStateMachines.Add(FString(PluginIdentifier.GetIdentifyingString()), NewStateMachine);
-	InitStateMachine(NewStateMachine, MoveTemp(PluginIdentifier), ProtocolOptions);
+	const int64 Id = NextInstanceId++;
+	GameFeaturePluginStateMachines.Add(Id, NewStateMachine);
+	GameFeaturesByIdentifier.Add(FString(PluginIdentifier.GetIdentifyingString()), NewStateMachine);
+	InitStateMachine(NewStateMachine, Id, MoveTemp(PluginIdentifier), ProtocolOptions);
 
 	return NewStateMachine;
 }
@@ -4002,7 +4015,8 @@ void UGameFeaturesSubsystem::BeginTermination(UGameFeaturePluginStateMachine* Ma
 	FStringView Identifer = GetPluginIdentifier(Machine).GetIdentifyingString();
 
 	UE_LOG(LogGameFeatures, Verbose, TEXT("BeginTermination of GameFeaturePlugin. Identifier:%.*s URL:%s"), Identifer.Len(), Identifer.GetData(), *GetPluginURL(Machine));
-	GameFeaturePluginStateMachines.RemoveByHash(GetTypeHash(Identifer), Identifer);
+	GameFeaturePluginStateMachines.Remove(Machine->StateProperties.Id);
+	GameFeaturesByIdentifier.RemoveByHash(GetTypeHash(Identifer), Identifer);
 	TerminalGameFeaturePluginStateMachines.Add(Machine);
 }
 
@@ -4010,6 +4024,7 @@ void UGameFeaturesSubsystem::FinishTermination(UGameFeaturePluginStateMachine* M
 {
 	UE_LOG(LogGameFeatures, Verbose, TEXT("FinishTermination of GameFeaturePlugin. Identifier:%.*s URL:%s"), GetPluginIdentifier(Machine).GetIdentifyingString().Len(), GetPluginIdentifier(Machine).GetIdentifyingString().GetData(), *GetPluginURL(Machine));
 	TerminalGameFeaturePluginStateMachines.RemoveSwap(Machine);
+	Machine->Executor = nullptr;
 }
 
 bool UGameFeaturesSubsystem::FindOrCreatePluginDependencyStateMachines(const FString& PluginURL, const FGameFeaturePluginStateMachineProperties& InStateProperties, TArray<UGameFeaturePluginStateMachine*>& OutDependencyMachines)
@@ -4548,6 +4563,11 @@ bool UGameFeaturesSubsystem::Tick(float DeltaTime)
 	TRACE_CPUPROFILER_EVENT_SCOPE(GFP_GameFeaturesSubsystem_Tick);
 
 	bool bContinueTicking = TickBatchProcessing();
+	if (StateMachineExecutor)
+	{
+		bContinueTicking = StateMachineExecutor->UpdateStateMachines() || bContinueTicking;
+	}
+	
 	if (!bContinueTicking)
 	{
 		DisableTick();
@@ -4621,7 +4641,7 @@ TSet<FString> UGameFeaturesSubsystem::GetActivePluginNames() const
 {
 	TSet<FString> ActivePluginNames;
 
-	for (const TPair<FString, TObjectPtr<UGameFeaturePluginStateMachine>>& Pair : GameFeaturePluginStateMachines)
+	for (const TPair<int64, TObjectPtr<UGameFeaturePluginStateMachine>>& Pair : GameFeaturePluginStateMachines)
 	{
 		UGameFeaturePluginStateMachine* StateMachine = Pair.Value;
 		if (GetCurrentState(StateMachine) == EGameFeaturePluginState::Active &&
diff --git a/Engine/Plugins/Runtime/GameFeatures/Source/GameFeatures/Public/GameFeatureTypes.h b/Engine/Plugins/Runtime/GameFeatures/Source/GameFeatures/Public/GameFeatureTypes.h
--- a/Engine/Plugins/Runtime/GameFeatures/Source/GameFeatures/Public/GameFeatureTypes.h
+++ b/Engine/Plugins/Runtime/GameFeatures/Source/GameFeatures/Public/GameFeatureTypes.h
@@ -6,6 +6,7 @@
 #include "Misc/EnumClassFlags.h"
 #include "Misc/EnumRange.h"
 
+class FCbWriter;
 class FString;
 
 /** 
@@ -53,6 +54,8 @@ class FString;
 	XSTATE(Active,							NSLOCTEXT("GameFeatures", "ActiveStateDisplayName", "Active"))													/* Plugin is fully loaded and active. It is affecting the game.  */ 
 
 #define GAME_FEATURE_PLUGIN_STATE_ENUM(inEnum, inText) inEnum,
+// EGameFeaturePluginState defines the list of public states for game feature plugins.
+// Internally, GFP state machines may use additional states to manage transitions between public states
 namespace GameFeaturePluginStatePrivate
 {
 	enum EGameFeaturePluginState : uint8
@@ -63,6 +66,9 @@ namespace GameFeaturePluginStatePrivate
 }
 using EGameFeaturePluginState = GameFeaturePluginStatePrivate::EGameFeaturePluginState;
 #undef GAME_FEATURE_PLUGIN_STATE_ENUM
+ENUM_RANGE_BY_COUNT(EGameFeaturePluginState, EGameFeaturePluginState::MAX);
+
+GAMEFEATURES_API void SerializeForLog(FCbWriter& Writer, EGameFeaturePluginState State);
 
 namespace UE::GameFeatures
 {
diff --git a/Engine/Plugins/Runtime/GameFeatures/Source/GameFeatures/Public/GameFeaturesSubsystem.h b/Engine/Plugins/Runtime/GameFeatures/Source/GameFeatures/Public/GameFeaturesSubsystem.h
--- a/Engine/Plugins/Runtime/GameFeatures/Source/GameFeatures/Public/GameFeaturesSubsystem.h
+++ b/Engine/Plugins/Runtime/GameFeatures/Source/GameFeatures/Public/GameFeaturesSubsystem.h
@@ -9,6 +9,7 @@
 #include "GameFeaturePluginOperationResult.h"
 #include "Misc/TransactionallySafeRWLock.h"
 #include "Stats/StatsMisc.h"
+#include "Templates/PimplPtr.h"
 
 #include "GameFeaturesSubsystem.generated.h"
 
@@ -17,7 +18,11 @@
 namespace UE::GameFeatures
 {
 	class FPackageLoadTracker;
+	class FRealtimeMode;
+	class FStateMachineExecutor;
 	struct FResult;
+
+	bool TickUnloadPluginAssets(float DeltaTime);
 }
 
 class UGameFeaturePluginStateMachine;
@@ -472,6 +477,8 @@ class UGameFeaturesSubsystem : public UEngineSubsystem
 	GENERATED_BODY()
 
 public:
+	UGameFeaturesSubsystem(const FObjectInitializer& ObjectInitializer);
+
 	//~UEngineSubsystem interface
 	UE_API virtual void Initialize(FSubsystemCollectionBase& Collection) override;
 	UE_API virtual void Deinitialize() override;
@@ -964,10 +971,22 @@ private:
 
 	UE_API bool TickBatchProcessing();
 
+	friend class UE::GameFeatures::FRealtimeMode;
+	friend bool UE::GameFeatures::TickUnloadPluginAssets(float DeltaTime);
+	void TryUpdateStateMachineImmediate(int64 Id);
+
 private:
-	/** The list of all game feature plugin state machine objects */
+	/** The list of all game feature plugin state machine objects addressed by FGameFeaturePluginStateMachineProperties::Id */
+	UPROPERTY(Transient)
+	TMap<int64, TObjectPtr<UGameFeaturePluginStateMachine>> GameFeaturePluginStateMachines;
+
+	/** Lookup of game feature plugins by FGameFeaturePluginIdentifier::GetIdentifyingString */
 	UPROPERTY(Transient)
-	TMap<FString, TObjectPtr<UGameFeaturePluginStateMachine>> GameFeaturePluginStateMachines;
+	TMap<FString, TObjectPtr<UGameFeaturePluginStateMachine>> GameFeaturesByIdentifier;
+
+	/** Private type controlling execution of state machine transitions when UE::GameFeatures::bUseNewExecutor */
+	TPimplPtr<UE::GameFeatures::FStateMachineExecutor> StateMachineExecutor;
+	friend class UE::GameFeatures::FStateMachineExecutor;
 
 	/** The tick handle if currently registered for a tick */
 	FTSTicker::FDelegateHandle TickHandle;
@@ -1018,6 +1037,11 @@ private:
 
 	static UE_API TSet<FString> DlcPlugins;
 
+	/**
+	 * Next identifier to be used for a GFP state machine.
+	 * Identifiers are not reused, and when a GFP is terminated and later recreated, a new identifier is assigned.
+	 */
+	int64 NextInstanceId = 0;
 
 private:
 	/**
@@ -1028,7 +1052,7 @@ private:
 	 */
 
 	/** Initializes the state machine and assigns the URL for the plugin it manages. This sets the machine to the 'UnknownStatus' state. */
-	void InitStateMachine(UGameFeaturePluginStateMachine* This, FGameFeaturePluginIdentifier InPluginIdentifier, const FGameFeatureProtocolOptions& InProtocolOptions);
+	void InitStateMachine(UGameFeaturePluginStateMachine* This, int64 Id, FGameFeaturePluginIdentifier InPluginIdentifier, const FGameFeatureProtocolOptions& InProtocolOptions);
 
 	/** Asynchronously transitions the state machine to the destination state range and reports when it is done. 
 	  * DestinationState must be of type EGameFeaturePluginStateType::Destination.
@@ -1090,7 +1114,7 @@ private:
 	FGameFeaturePluginStateRange GetDestination(const UGameFeaturePluginStateMachine* This) const;
 
 	/** Returns information about the current state */
-	const FGameFeaturePluginStateInfo& GetCurrentStateInfo(const UGameFeaturePluginStateMachine* This) const;
+	FGameFeaturePluginStateInfo GetCurrentStateInfo(const UGameFeaturePluginStateMachine* This) const;
 
 	/** Returns true if attempting to reach a new destination state */
 	bool IsRunning(const UGameFeaturePluginStateMachine* This) const;
@@ -1147,10 +1171,7 @@ private:
 	bool IsValidErrorState(const UGameFeaturePluginStateMachine* This, EGameFeaturePluginState InDestinationState) const;
 
 	/** Processes the current state and looks for state transitions */
-	void UpdateStateMachine(UGameFeaturePluginStateMachine* This);
-
-	/** Update Progress for current state */
-	void UpdateCurrentStateProgress(UGameFeaturePluginStateMachine* This, float Progress) const;
+	void UpdateStateMachine(UGameFeaturePluginStateMachine* This, TOptional<int32> MaxTransitions = {});
 
 	/** Helper to register the state machine with the manager as a state machine in transiton */
 	void RegisterAsTransitioningStateMachine(UGameFeaturePluginStateMachine* This);
