If we are downgrading ourselfs, never go below our BuiltinState if we are a Builtin GFP

#Synced-CL 50191948

diff --git a/Engine/Plugins/Runtime/GameFeatures/Source/GameFeatures/Private/GameFeatureStateHandleReferenceController.cpp b/Engine/Plugins/Runtime/GameFeatures/Source/GameFeatures/Private/GameFeatureStateHandleReferenceController.cpp
--- a/Engine/Plugins/Runtime/GameFeatures/Source/GameFeatures/Private/GameFeatureStateHandleReferenceController.cpp
+++ b/Engine/Plugins/Runtime/GameFeatures/Source/GameFeatures/Private/GameFeatureStateHandleReferenceController.cpp
@@ -381,6 +381,20 @@ bool FGameFeatureStateHandleReferenceController::GetDowngradedTargetState(const
 	}
 	else if (RefCount->GetHighestRefCountDestPluginState(OutTransitionState))
 	{
+		// if we were loaded builtin, dont let us downgrade past our builtin state
+		if (GameFeatures.WasGameFeaturePluginLoadedAsBuiltIn(PluginURL))
+		{
+			FGameFeaturePluginDetails PluginDetails;
+			if (GameFeatures.GetGameFeaturePluginDetails(PluginURL, PluginDetails))
+			{
+				EGameFeaturePluginState BuiltinState = UGameFeaturesSubsystem::ConvertInitialFeatureStateToTargetState(PluginDetails.BuiltInAutoState);
+				if (BuiltinState > OutTransitionState)
+				{
+					OutTransitionState = BuiltinState;
+				}
+			}
+		}
+
 		UE_LOG(LogGameFeatureStateHandle, Verbose, TEXT("  depends '%s' downgrading to -> (%s)"), *GetPluginNameFromURL(PluginURL), *UE::GameFeatures::ToString(OutTransitionState));
 		return true;
 	}
