Allow another StateHandle to take ownership of another StateHandle

#Synced-CL 50193382

#Synced-CL 50193382
#Synced-CL 50193382

diff --git a/Engine/Plugins/Runtime/GameFeatures/Source/GameFeatures/Private/GameFeatureStateHandle.cpp b/Engine/Plugins/Runtime/GameFeatures/Source/GameFeatures/Private/GameFeatureStateHandle.cpp
--- a/Engine/Plugins/Runtime/GameFeatures/Source/GameFeatures/Private/GameFeatureStateHandle.cpp
+++ b/Engine/Plugins/Runtime/GameFeatures/Source/GameFeatures/Private/GameFeatureStateHandle.cpp
@@ -26,6 +26,17 @@ void FGameFeatureStateHandle::Invalidate()
 	UniqueId.Invalidate();
 }
 
+void FGameFeatureStateHandle::TakeOwnership(FGameFeatureStateHandle& FromStateHandle)
+{
+	if (*this == FromStateHandle)
+	{
+		return;
+	}
+
+	FGameFeatureStateHandleReferenceController::Get().MergeGameFeatureStateHandle(FromStateHandle, *this);
+	FGameFeatureStateHandleReferenceController::Get().UnregisterStateHandle(FromStateHandle); // removes from the internal list + invalidates the StateHandle
+}
+
 FString FGameFeatureStateHandle::GetOwner() const
 {
 	return Owner;
@@ -67,6 +78,12 @@ FGameFeatureStateHandleInternal::~FGameFeatureStateHandleInternal()
 	// ensureMsgf(StateHandleData.IsEmpty(), TEXT("[FGameFeatureStateHandle] We are leaking %i GFPs and need to call \"ResetGameFeatureStateHandle\" to clear these ref counts up"), StateHandleData.Num());
 }
 
+void FGameFeatureStateHandleInternal::TakeOwnership(FGameFeatureStateHandleInternal& StateHandle)
+{
+	StateHandleData.Append(StateHandle.StateHandleData);
+	StateHandle.Empty();
+}
+
 void FGameFeatureStateHandleInternal::AddOrUpdateOwnershipIfHighestDestState(const FString& PluginURL, EGameFeaturePluginState DestPluginState)
 {
 	EGameFeaturePluginState& PluginState = StateHandleData.FindOrAdd(PluginURL);
diff --git a/Engine/Plugins/Runtime/GameFeatures/Source/GameFeatures/Private/GameFeatureStateHandleInternal.h b/Engine/Plugins/Runtime/GameFeatures/Source/GameFeatures/Private/GameFeatureStateHandleInternal.h
--- a/Engine/Plugins/Runtime/GameFeatures/Source/GameFeatures/Private/GameFeatureStateHandleInternal.h
+++ b/Engine/Plugins/Runtime/GameFeatures/Source/GameFeatures/Private/GameFeatureStateHandleInternal.h
@@ -24,6 +24,8 @@ public:
 	FGameFeatureStateHandleInternal(const FGameFeatureStateHandleInternal& OtherStateHandle) = delete;
 	FGameFeatureStateHandleInternal& operator=(const FGameFeatureStateHandleInternal& OtherStateHandle) = delete;
 
+	void TakeOwnership(FGameFeatureStateHandleInternal& StateHandle);
+
 	void AddOrUpdateOwnershipIfHighestDestState(const FString& PluginURL, EGameFeaturePluginState DestPluginState);
 
 	const TArray<FString> GetPlugins() const;
diff --git a/Engine/Plugins/Runtime/GameFeatures/Source/GameFeatures/Private/GameFeatureStateHandleReferenceController.cpp b/Engine/Plugins/Runtime/GameFeatures/Source/GameFeatures/Private/GameFeatureStateHandleReferenceController.cpp
--- a/Engine/Plugins/Runtime/GameFeatures/Source/GameFeatures/Private/GameFeatureStateHandleReferenceController.cpp
+++ b/Engine/Plugins/Runtime/GameFeatures/Source/GameFeatures/Private/GameFeatureStateHandleReferenceController.cpp
@@ -697,6 +697,26 @@ void FGameFeatureStateHandleReferenceController::ResetSingleGameFeatureStateHand
 	}
 }
 
+void FGameFeatureStateHandleReferenceController::MergeGameFeatureStateHandle(const FGameFeatureStateHandle& From, const FGameFeatureStateHandle& To)
+{
+	if (From == To)
+	{
+		return;
+	}
+
+	FScopeLock Lock(&Internal->StateHandlesLock);
+	FGameFeatureStateHandleInternal* FromStateHandle = Internal->StateHandles.Find(From);
+	FGameFeatureStateHandleInternal* ToStateHandle   = Internal->StateHandles.Find(To);
+
+	if (FromStateHandle == nullptr || ToStateHandle == nullptr)
+	{
+		return;
+	}
+
+	// this will leave FromStateHandle with no GFPs it owns
+	ToStateHandle->TakeOwnership(*FromStateHandle);
+}
+
 void FGameFeatureStateHandleReferenceControllerInternal::ListGameFeatureStateHandles(const TArray<FString>& Args, UWorld* World, FOutputDevice& Ar)
 {
 	FScopeLock Lock(&StateHandlesLock);
diff --git a/Engine/Plugins/Runtime/GameFeatures/Source/GameFeatures/Public/GameFeatureStateHandle.h b/Engine/Plugins/Runtime/GameFeatures/Source/GameFeatures/Public/GameFeatureStateHandle.h
--- a/Engine/Plugins/Runtime/GameFeatures/Source/GameFeatures/Public/GameFeatureStateHandle.h
+++ b/Engine/Plugins/Runtime/GameFeatures/Source/GameFeatures/Public/GameFeatureStateHandle.h
@@ -34,6 +34,8 @@ public:
 	UE_API bool TrackDependencies() const;
 	UE_API bool KeepRegistered() const;
 
+	UE_API void TakeOwnership(FGameFeatureStateHandle& StateHandle);
+
 	UE_API bool operator==(const FGameFeatureStateHandle& Other) const;
 	friend uint32 GetTypeHash(const FGameFeatureStateHandle& StateHandle);
 
diff --git a/Engine/Plugins/Runtime/GameFeatures/Source/GameFeatures/Public/GameFeatureStateHandleReferenceController.h b/Engine/Plugins/Runtime/GameFeatures/Source/GameFeatures/Public/GameFeatureStateHandleReferenceController.h
--- a/Engine/Plugins/Runtime/GameFeatures/Source/GameFeatures/Public/GameFeatureStateHandleReferenceController.h
+++ b/Engine/Plugins/Runtime/GameFeatures/Source/GameFeatures/Public/GameFeatureStateHandleReferenceController.h
@@ -47,6 +47,9 @@ public:
 	/** From the RefCounts gets the lowest GFP state required by a GFP we are tracking into OutPluginState */
 	UE_API bool GetMinimumRequiredStateForGameFeaturePlugin(const FString& GameFeaturePluginURL, EGameFeaturePluginState& OutPluginState);
 
+	/** The From StateHandles internals are appended to the To StateHandle, and the list of GFPs in From are Emptied */
+	UE_API void MergeGameFeatureStateHandle(const FGameFeatureStateHandle& From, const FGameFeatureStateHandle& To);
+
 private:
 	FGameFeatureStateHandleReferenceController();
 	FGameFeatureStateHandleReferenceController(FGameFeatureStateHandleReferenceController const&) = delete;
