Fix move semantics not invalidating the unique id and causing use
Fix possibly crashing on exit if the not all the StateHandles have been cleaned up and we are requesting exit and hitting our Singleton that may be gone when Unregistering

#Synced-CL 50193382

diff --git a/Engine/Plugins/Runtime/GameFeatures/Source/GameFeatures/Private/GameFeatureStateHandle.cpp b/Engine/Plugins/Runtime/GameFeatures/Source/GameFeatures/Private/GameFeatureStateHandle.cpp
--- a/Engine/Plugins/Runtime/GameFeatures/Source/GameFeatures/Private/GameFeatureStateHandle.cpp
+++ b/Engine/Plugins/Runtime/GameFeatures/Source/GameFeatures/Private/GameFeatureStateHandle.cpp
@@ -10,6 +10,25 @@ FGameFeatureStateHandle::~FGameFeatureStateHandle()
 	ReleaseAndUnregister();
 }
 
+FGameFeatureStateHandle::FGameFeatureStateHandle(FGameFeatureStateHandle&& Other)
+	: UniqueId(Other.UniqueId)
+{
+	Other.UniqueId.Invalidate();
+}
+
+FGameFeatureStateHandle& FGameFeatureStateHandle::operator=(FGameFeatureStateHandle&& Other)
+{
+	if (this != &Other)
+	{
+		ReleaseAndUnregister();
+
+		UniqueId = Other.UniqueId;
+		Other.UniqueId.Invalidate();
+	}
+
+	return *this;
+}
+
 bool FGameFeatureStateHandle::InitAndRegister(const FString& Owner, EGameFeatureStateHandleOptions Options)
 {
 	if (IsValid())
diff --git a/Engine/Plugins/Runtime/GameFeatures/Source/GameFeatures/Private/GameFeatureStateHandleReferenceController.cpp b/Engine/Plugins/Runtime/GameFeatures/Source/GameFeatures/Private/GameFeatureStateHandleReferenceController.cpp
--- a/Engine/Plugins/Runtime/GameFeatures/Source/GameFeatures/Private/GameFeatureStateHandleReferenceController.cpp
+++ b/Engine/Plugins/Runtime/GameFeatures/Source/GameFeatures/Private/GameFeatureStateHandleReferenceController.cpp
@@ -173,10 +173,13 @@ void FGameFeatureStateHandleReferenceController::RegisterNewStateHandle(const FG
 
 void FGameFeatureStateHandleReferenceController::UnregisterStateHandle(const FGameFeatureStateHandle& StateHandle)
 {
-	FScopeLock Lock(&Internal->StateHandlesLock);
-	UE_LOG(LogGameFeatureStateHandle, Verbose, TEXT("Unregistereing and invalidating StateHandle: %s %i"), *StateHandle.ToString(), Internal->StateHandles.Num());
+	if (!IsEngineExitRequested())
+	{
+		FScopeLock Lock(&Internal->StateHandlesLock);
+		UE_LOG(LogGameFeatureStateHandle, Verbose, TEXT("Unregistereing and invalidating StateHandle: %s %i"), *StateHandle.ToString(), Internal->StateHandles.Num());
 
-	Internal->StateHandles.Remove(StateHandle.GetUniqueId());
+		Internal->StateHandles.Remove(StateHandle.GetUniqueId());
+	}
 }
 
 // TODO possibly its faster/better to write an EnumerateOverDepends + a lambda callback to avoid storing a data structure to return
diff --git a/Engine/Plugins/Runtime/GameFeatures/Source/GameFeatures/Public/GameFeatureStateHandle.h b/Engine/Plugins/Runtime/GameFeatures/Source/GameFeatures/Public/GameFeatureStateHandle.h
--- a/Engine/Plugins/Runtime/GameFeatures/Source/GameFeatures/Public/GameFeatureStateHandle.h
+++ b/Engine/Plugins/Runtime/GameFeatures/Source/GameFeatures/Public/GameFeatureStateHandle.h
@@ -36,8 +36,8 @@ public:
 	FGameFeatureStateHandle(const FGameFeatureStateHandle&) = delete;
 	FGameFeatureStateHandle& operator=(const FGameFeatureStateHandle&) = delete;
 
-	FGameFeatureStateHandle(FGameFeatureStateHandle&&) = default;
-	FGameFeatureStateHandle& operator=(FGameFeatureStateHandle&&) = default;
+	UE_API FGameFeatureStateHandle(FGameFeatureStateHandle&&);
+	UE_API FGameFeatureStateHandle& operator=(FGameFeatureStateHandle&&);
 
 	/** Init the StateHandle with a unique id, owner, options and registers it with the StateHandleRefController, if this StateHandle already has a valid GUID, we return false */
 	UE_API bool InitAndRegister(const FString& Owner, EGameFeatureStateHandleOptions Options = EGameFeatureStateHandleOptions::TrackDependencies);
