Remove the Reset One GFP from StateHandles, with depend tracking this can leave our ref counts in the incorrect state
If you need to reset a single GFP from the StateHandle then it should be done currently from a list of state handles owning a single GFP that can be removed/reset

#Synced-CL 50193382

diff --git a/Engine/Plugins/Runtime/GameFeatures/Source/GameFeatures/Private/GameFeatureStateHandleReferenceController.cpp b/Engine/Plugins/Runtime/GameFeatures/Source/GameFeatures/Private/GameFeatureStateHandleReferenceController.cpp
--- a/Engine/Plugins/Runtime/GameFeatures/Source/GameFeatures/Private/GameFeatureStateHandleReferenceController.cpp
+++ b/Engine/Plugins/Runtime/GameFeatures/Source/GameFeatures/Private/GameFeatureStateHandleReferenceController.cpp
@@ -703,34 +703,6 @@ void FGameFeatureStateHandleReferenceController::ResetGameFeatureStateHandle(con
 	}
 }
 
-void FGameFeatureStateHandleReferenceController::ResetSingleGameFeatureStateHandle(const FGameFeatureStateHandle& StateHandle, const FString& PluginURL, TFunction<void(bool)> OnComplete)
-{
-	if (!IsLifetimeControlEnabled())
-	{
-		OnComplete(true);
-	}
-	else if (IsEngineExitRequested())
-	{
-		UE_LOG(LogGameFeatureStateHandle, Verbose, TEXT("Requesting a reset on %s but the Engine is Exiting. This should have likely been done before we existed the Engine. StateHandle Owner '%s'"), *PluginURL, *StateHandle.ToString());
-		OnComplete(true);
-	}
-	else
-	{
-		FScopeLock Lock(&Internal->StateHandlesLock);
-		if (FGameFeatureStateHandleInternal* InternalHandle = Internal->StateHandles.Find(StateHandle.GetUniqueId()))
-		{
-			FString PluginName = GetPluginNameFromURL(PluginURL);
-			InternalHandle->Remove(PluginName);
-
-			RemoveReferences(StateHandle, TArray<FString>{ PluginName }, OnComplete);
-		}
-		else
-		{
-			OnComplete(false);
-		}
-	}
-}
-
 void FGameFeatureStateHandleReferenceController::MergeGameFeatureStateHandle(const FGameFeatureStateHandle& From, const FGameFeatureStateHandle& To)
 {
 	if (From == To)
diff --git a/Engine/Plugins/Runtime/GameFeatures/Source/GameFeatures/Public/GameFeatureStateHandleReferenceController.h b/Engine/Plugins/Runtime/GameFeatures/Source/GameFeatures/Public/GameFeatureStateHandleReferenceController.h
--- a/Engine/Plugins/Runtime/GameFeatures/Source/GameFeatures/Public/GameFeatureStateHandleReferenceController.h
+++ b/Engine/Plugins/Runtime/GameFeatures/Source/GameFeatures/Public/GameFeatureStateHandleReferenceController.h
@@ -38,9 +38,6 @@ public:
 	/** Removes all the GFPs references held by this list of handles, and will empty out the StateHandle list */
 	UE_API void ResetGameFeatureStateHandle(const FGameFeatureStateHandle& StateHandle, TFunction<void(bool)> OnComplete = [](bool) {});
 
-	/** Removes a single GFP references held by this list of handles, and will empty out the StateHandle list */
-	UE_API void ResetSingleGameFeatureStateHandle(const FGameFeatureStateHandle& StateHandle, const FString& PluginURL, TFunction<void(bool)> OnComplete = [](bool) {});
-
 	/** Adds or Updates the GFP URL to the PluginState */
 	UE_API void AddOrUpdateReference(const FGameFeatureStateHandle& StateHandle, const FString& GameFeaturePluginURL, EGameFeaturePluginState PluginState, bool bCollectDepends = true);
 
