Fix missing Owner name now that the internal state handle owns that
Avoid calling singletons, even if in the same module, if we are requesting exit

#Synced-CL 50193382

diff --git a/Engine/Plugins/Runtime/GameFeatures/Source/GameFeatures/Private/GameFeatureStateHandle.cpp b/Engine/Plugins/Runtime/GameFeatures/Source/GameFeatures/Private/GameFeatureStateHandle.cpp
--- a/Engine/Plugins/Runtime/GameFeatures/Source/GameFeatures/Private/GameFeatureStateHandle.cpp
+++ b/Engine/Plugins/Runtime/GameFeatures/Source/GameFeatures/Private/GameFeatureStateHandle.cpp
@@ -47,7 +47,7 @@ bool FGameFeatureStateHandle::InitAndRegister(const FString& Owner, EGameFeature
 
 void FGameFeatureStateHandle::ReleaseAndUnregister(TFunction<void(bool)> ResetCompleteCallback)
 {
-	if (IsValid())
+	if (IsValid() && !IsEngineExitRequested())
 	{
 		FGameFeatureStateHandleReferenceController::Get().ResetGameFeatureStateHandle(*this, ResetCompleteCallback);
 		FGameFeatureStateHandleReferenceController::Get().UnregisterStateHandle(*this);
diff --git a/Engine/Plugins/Runtime/GameFeatures/Source/GameFeatures/Private/GameFeatureStateHandleReferenceController.cpp b/Engine/Plugins/Runtime/GameFeatures/Source/GameFeatures/Private/GameFeatureStateHandleReferenceController.cpp
--- a/Engine/Plugins/Runtime/GameFeatures/Source/GameFeatures/Private/GameFeatureStateHandleReferenceController.cpp
+++ b/Engine/Plugins/Runtime/GameFeatures/Source/GameFeatures/Private/GameFeatureStateHandleReferenceController.cpp
@@ -441,13 +441,26 @@ void FGameFeatureStateHandleReferenceController::RemoveReferences(const FGameFea
 		}
 	}
 
+	FString StateHandleOwner;
+	{
+		FScopeLock Lock(&Internal->StateHandlesLock);
+		if (FGameFeatureStateHandleInternal* InternalHandle = Internal->StateHandles.Find(StateHandle.GetUniqueId()))
+		{
+			StateHandleOwner = InternalHandle->GetOwner();
+		}
+		else
+		{
+			StateHandleOwner = StateHandle.ToString();
+		}
+	}
+
 	if (PluginsToDowngrade.IsEmpty())
 	{
-		UE_LOG(LogGameFeatureStateHandle, Verbose, TEXT("Removing References and Downgrading %i/%i GFPs that was owned from '%s'"), PluginsToDowngrade.Num(), PluginsToRemoveRefences.Num(), *StateHandle.ToString());
+		UE_LOG(LogGameFeatureStateHandle, Verbose, TEXT("Removing References and Downgrading %i/%i GFPs that was owned from '%s'"), PluginsToDowngrade.Num(), PluginsToRemoveRefences.Num(), *StateHandleOwner);
 	}
 	else
 	{
-		UE_LOG(LogGameFeatureStateHandle, Display, TEXT("Removing References and Downgrading %i/%i GFPs that was owned from '%s'"), PluginsToDowngrade.Num(), PluginsToRemoveRefences.Num(), *StateHandle.ToString());
+		UE_LOG(LogGameFeatureStateHandle, Display, TEXT("Removing References and Downgrading %i/%i GFPs that was owned from '%s'"), PluginsToDowngrade.Num(), PluginsToRemoveRefences.Num(), *StateHandleOwner);
 	}
 
 	// return early if we have no downgrades to do
