UE: Add an option when calling PredownloadGameFeaturePlugins to wait to call OnComplete until the bundle release is complete. This is useful if you want to be sure all deferred tasks are complete before proceeding. Also fixed a bug where the FGameFeaturePluginPredownloadContext was not properly unbinding from InstallBundleCompleteDelegate and thus never being deleted.

diff --git a/Engine/Plugins/Runtime/GameFeatures/Source/GameFeatures/Private/GameFeaturesSubsystem.cpp b/Engine/Plugins/Runtime/GameFeatures/Source/GameFeatures/Private/GameFeaturesSubsystem.cpp
--- a/Engine/Plugins/Runtime/GameFeatures/Source/GameFeatures/Private/GameFeaturesSubsystem.cpp
+++ b/Engine/Plugins/Runtime/GameFeatures/Source/GameFeatures/Private/GameFeaturesSubsystem.cpp
@@ -3235,6 +3235,7 @@ struct FGameFeaturePluginPredownloadContext : public FGameFeaturePluginPredownlo
 	TMap<FGameFeaturePluginIdentifier, FGFPData> GFPs;
 
 	TArray<FName> PendingBundleDownloads;
+	TArray<FName> PendingBundleReleases;
 
 	FInstallBundleCombinedProgressTracker ProgressTracker{ false /*bAutoTick*/};
 
@@ -3252,6 +3253,7 @@ struct FGameFeaturePluginPredownloadContext : public FGameFeaturePluginPredownlo
 	UE::FMutex IoStoreMutex;
 
 	bool bHasAssetDependencies = false;
+	bool bWaitForBundlesToRelease = false;
 	bool bIsComplete = false;
 	bool bCanceled = false;
 
@@ -3307,12 +3309,15 @@ struct FGameFeaturePluginPredownloadContext : public FGameFeaturePluginPredownlo
 		}		
 		IInstallBundleManager::InstallBundleCompleteDelegate.RemoveAll(this);
 		IInstallBundleManager::PausedBundleDelegate.RemoveAll(this);
+		if (bWaitForBundlesToRelease && PendingBundleReleases.Num() == 0)
+		{
+			// If we have pending releases, this will be cleaned up once the last one is released
+			IInstallBundleManager::ReleasedDelegate.RemoveAll(this);
+		}
 	}
 
-	void SetComplete()
+	void FinishComplete()
 	{
-		ReleaseBundlesIfPossible();
-
 		bIsComplete = true;
 		if (OnComplete)
 		{
@@ -3320,15 +3325,24 @@ struct FGameFeaturePluginPredownloadContext : public FGameFeaturePluginPredownlo
 		}
 	}
 
+	void SetComplete()
+	{
+		ReleaseBundlesIfPossible();
+
+		if (PendingBundleReleases.Num() == 0)
+		{
+			FinishComplete();
+		}
+	}
+
 	void SetComplete(UE::GameFeatures::FResult&& InResult)
 	{
 		ReleaseBundlesIfPossible();
 
 		Result = MoveTemp(InResult);
-		bIsComplete = true;
-		if (OnComplete)
+		if (PendingBundleReleases.Num() == 0)
 		{
-			OnComplete(Result);
+			FinishComplete();
 		}
 	}
 
@@ -3339,10 +3353,9 @@ struct FGameFeaturePluginPredownloadContext : public FGameFeaturePluginPredownlo
 		Result = MakeError(FString::Printf(TEXT("%.*s%s"),
 			PredownloadErrorNamespace.Len(), PredownloadErrorNamespace.GetData(),
 			TEXT("Canceled")));
-		bIsComplete = true;
-		if (OnComplete)
+		if (PendingBundleReleases.Num() == 0)
 		{
-			OnComplete(Result);
+			FinishComplete();
 		}
 	}
 
@@ -3544,7 +3557,7 @@ struct FGameFeaturePluginPredownloadContext : public FGameFeaturePluginPredownlo
 		// Start ticking
 		TickHandle = FTSTicker::GetCoreTicker().AddTicker(FTickerDelegate::CreateSP(this, &FGameFeaturePluginPredownloadContext::Tick));
 
-		IInstallBundleManager::InstallBundleCompleteDelegate.AddLambda(
+		IInstallBundleManager::InstallBundleCompleteDelegate.AddSPLambda(this,
 			[Context = SharedThis(this)](FInstallBundleRequestResultInfo BundleResult)
 			{ Context->OnInstallBundleCompleted(MoveTemp(BundleResult)); });
 		// TODO: handle pause?  Just cancel? This should only be relevent for cell connections
@@ -3601,6 +3614,29 @@ struct FGameFeaturePluginPredownloadContext : public FGameFeaturePluginPredownlo
 		);
 	}
 
+	void OnBundleReleaseCompleted(FInstallBundleReleaseRequestResultInfo ResultInfo)
+	{
+		// Note: ReleaseBundlesIfPossible is what populates PendingBundleReleases and it is only called as a result of SetComplete.
+		// For this reason, we will assume that once we are done releasing then we are complete.
+		if (PendingBundleReleases.Num() > 0)
+		{
+			PendingBundleReleases.Remove(ResultInfo.BundleName);
+			
+			if (PendingBundleReleases.Num() == 0)
+			{
+				UE_LOGFMT(LogGameFeatures, Display, "GFP Predownload has finished releasing all bundles. Completed.");
+				FinishComplete();
+
+				// Since we did not remove the delegate handle earlier when we called SetComplete, remove it now
+				IInstallBundleManager::ReleasedDelegate.RemoveAll(this);
+
+				// NOTE: After we remove ourselves from ReleasedDelegate, we will likely be deleted.
+				// Do not add any code after this line. Explicitly returning here to protect against future code changes
+				return;
+			}
+		}
+	}
+
 	void OnAllInstallBundlesCompleted()
 	{
 		if (Result.HasError() || !bHasAssetDependencies)
@@ -3836,15 +3872,31 @@ struct FGameFeaturePluginPredownloadContext : public FGameFeaturePluginPredownlo
 			ReleaseList.Append(GFPData.ProtocolMetadata.InstallBundles);
 		}
 
+		if (bWaitForBundlesToRelease && ReleaseList.Num() > 0)
+		{
+			FString BundlesStr = ReleaseList[0].ToString();
+			for (int32 BundleIdx = 1; BundleIdx < ReleaseList.Num(); ++BundleIdx)
+			{
+				BundlesStr += TEXT(",") + ReleaseList[BundleIdx].ToString();
+			}
+			UE_LOGFMT(LogGameFeatures, Display, "GFP Predownload is waiting to release bundles before completing. Bundles: {Bundles}", ("Bundles", BundlesStr));
+			PendingBundleReleases = ReleaseList;
+
+			IInstallBundleManager::ReleasedDelegate.AddSPLambda(this,
+				[Context = SharedThis(this)](FInstallBundleReleaseRequestResultInfo ResultInfo)
+				{ Context->OnBundleReleaseCompleted(MoveTemp(ResultInfo)); });
+		}
+
 		BundleManager->RequestReleaseContent(ReleaseList, EInstallBundleReleaseRequestFlags::None, KeepList.Array());
 	}
 };
 
-TSharedRef<FGameFeaturePluginPredownloadHandle> UGameFeaturesSubsystem::PredownloadGameFeaturePlugins(TConstArrayView<FString> PluginURLs, TUniqueFunction<void(const UE::GameFeatures::FResult&)> OnComplete /*= nullptr*/, TUniqueFunction<void(float)> OnProgress /*= nullptr*/)
+TSharedRef<FGameFeaturePluginPredownloadHandle> UGameFeaturesSubsystem::PredownloadGameFeaturePlugins(TConstArrayView<FString> PluginURLs, TUniqueFunction<void(const UE::GameFeatures::FResult&)> OnComplete /*= nullptr*/, TUniqueFunction<void(float)> OnProgress /*= nullptr*/, bool bWaitForBundlesToRelease /*= false*/)
 {
 	TSharedRef<FGameFeaturePluginPredownloadContext> Context = MakeShared<FGameFeaturePluginPredownloadContext>();
 	Context->OnComplete = MoveTemp(OnComplete);
 	Context->OnProgress = MoveTemp(OnProgress);
+	Context->bWaitForBundlesToRelease = bWaitForBundlesToRelease;
 	Context->Start(PluginURLs);
 
 	return Context;
diff --git a/Engine/Plugins/Runtime/GameFeatures/Source/GameFeatures/Public/GameFeaturesSubsystem.h b/Engine/Plugins/Runtime/GameFeatures/Source/GameFeatures/Public/GameFeaturesSubsystem.h
--- a/Engine/Plugins/Runtime/GameFeatures/Source/GameFeatures/Public/GameFeaturesSubsystem.h
+++ b/Engine/Plugins/Runtime/GameFeatures/Source/GameFeatures/Public/GameFeaturesSubsystem.h
@@ -766,9 +766,10 @@ public:
 	 * Pre-install any required game feature data, which can be useful for larger payloads. 
 	 * This does not instantiate any GFP although it is safe to do so before this finishes. 
 	 * This doesn't not resolve any dependencies, PluginURLs should contain all dependencies.
+	 * If bWaitForBundlesToRelease is true, OnComplete will not be called until the bundles have been released
 	 */
 	UE_API TSharedRef<FGameFeaturePluginPredownloadHandle> PredownloadGameFeaturePlugins(TConstArrayView<FString> PluginURLs, 
-		TUniqueFunction<void(const UE::GameFeatures::FResult&)> OnComplete = nullptr, TUniqueFunction<void(float)> OnProgress = nullptr);
+		TUniqueFunction<void(const UE::GameFeatures::FResult&)> OnComplete = nullptr, TUniqueFunction<void(float)> OnProgress = nullptr, bool bWaitForBundlesToRelease = false);
 	friend struct FGameFeaturePluginPredownloadContext;
 
 	/** Determine the initial feature state for a built-in plugin */
