Add UGameFeaturesSubsystem::IsGameFeaturePlugin utility function to reduce code duplication
#rnx
#Synced-CL 50516840

diff --git a/Engine/Plugins/Runtime/GameFeatures/Source/GameFeatures/Private/GameFeaturesSubsystem.cpp b/Engine/Plugins/Runtime/GameFeatures/Source/GameFeatures/Private/GameFeaturesSubsystem.cpp
--- a/Engine/Plugins/Runtime/GameFeatures/Source/GameFeatures/Private/GameFeaturesSubsystem.cpp
+++ b/Engine/Plugins/Runtime/GameFeatures/Source/GameFeatures/Private/GameFeaturesSubsystem.cpp
@@ -4879,6 +4879,22 @@ const TSet<FString>& UGameFeaturesSubsystem::GetDlcPlugins()
 	return DlcPlugins;
 }
 
+bool UGameFeaturesSubsystem::IsGameFeaturePlugin(const TSharedPtr<IPlugin>& InPlugin)
+{
+	const FString PluginsFolderRoot = (InPlugin->GetLoadedFrom() == EPluginLoadedFrom::Project) ? FPaths::ProjectPluginsDir() : FPaths::EnginePluginsDir();
+	FString PluginPathRelativeToDomain = InPlugin->GetBaseDir();
+	if (FPaths::MakePathRelativeTo(PluginPathRelativeToDomain, *PluginsFolderRoot))
+	{
+		PluginPathRelativeToDomain = TEXT("/") + PluginPathRelativeToDomain + TEXT("/");
+	}
+	else
+	{
+		PluginPathRelativeToDomain = FString();
+	}
+
+	return PluginPathRelativeToDomain.StartsWith(TEXT("/GameFeatures/"));
+}
+
 bool UGameFeaturesSubsystem::GetPluginDebugStateEnabled(const FString& PluginUrl) const
 {
 #if !UE_BUILD_SHIPPING
diff --git a/Engine/Plugins/Runtime/GameFeatures/Source/GameFeatures/Public/GameFeaturesSubsystem.h b/Engine/Plugins/Runtime/GameFeatures/Source/GameFeatures/Public/GameFeaturesSubsystem.h
--- a/Engine/Plugins/Runtime/GameFeatures/Source/GameFeatures/Public/GameFeaturesSubsystem.h
+++ b/Engine/Plugins/Runtime/GameFeatures/Source/GameFeatures/Public/GameFeaturesSubsystem.h
@@ -782,6 +782,13 @@ public:
 
 	static UE_API const TSet<FString>& GetDlcPlugins();
 
+	/**
+	 * Returns true if plugin is a GameFeaturePlugin
+	 * @param InPlugin The plugin to check
+	 * @return True if plugin considered a GameFeaturePlugin
+	 */
+	static UE_API bool IsGameFeaturePlugin(const TSharedPtr<IPlugin>& InPlugin);
+
 	UE_API bool GetPluginDebugStateEnabled(const FString& PluginUrl) const;
 	UE_API void SetPluginDebugStateEnabled(const FString& PluginUrl, bool bEnabled);
 
