[ContentBundle] Option for ContentBundle to EDL conversion commandlet to delete HLODs, as those can block the conversion

#tests Tested on Sparks

#Synced-CL 50538925

diff --git a/Engine/Plugins/Runtime/GameFeatures/Source/GameFeaturesEditor/Private/GameFeatureActionConvertContentBundleWorldPartitionBuilder.cpp b/Engine/Plugins/Runtime/GameFeatures/Source/GameFeaturesEditor/Private/GameFeatureActionConvertContentBundleWorldPartitionBuilder.cpp
--- a/Engine/Plugins/Runtime/GameFeatures/Source/GameFeaturesEditor/Private/GameFeatureActionConvertContentBundleWorldPartitionBuilder.cpp
+++ b/Engine/Plugins/Runtime/GameFeatures/Source/GameFeaturesEditor/Private/GameFeatureActionConvertContentBundleWorldPartitionBuilder.cpp
@@ -10,6 +10,7 @@
 #include "WorldPartition/DataLayer/ExternalDataLayerAsset.h"
 #include "WorldPartition/DataLayer/ExternalDataLayerInstance.h"
 #include "WorldPartition/DataLayer/ExternalDataLayerManager.h"
+#include "WorldPartition/HLOD/HLODActor.h"
 #include "DataLayer/DataLayerEditorSubsystem.h"
 #include "DataLayer/ExternalDataLayerFactory.h"
 #include "UObject/StrongObjectPtr.h"
@@ -33,6 +34,7 @@ UGameFeatureActionConvertContentBundleWorldPartitionBuilder::UGameFeatureActionC
 	: Super(ObjectInitializer)
 	, bReportOnly(false)
 	, bRemoveContentBundleAction(false)
+	, bDeleteHLODs(false)
 {}
 
 bool UGameFeatureActionConvertContentBundleWorldPartitionBuilder::PreRun(UWorld* World, FPackageSourceControlHelper& PackageHelper)
@@ -45,6 +47,7 @@ bool UGameFeatureActionConvertContentBundleWorldPartitionBuilder::PreRun(UWorld*
 	DestinationFolder = DestFolderPtr ? *DestFolderPtr : FString();
 	bReportOnly = Switches.Contains(TEXT("ReportOnly"));
 	bRemoveContentBundleAction = Switches.Contains(TEXT("RemoveContentBundleAction"));
+	bDeleteHLODs = Switches.Contains(TEXT("DeleteHLODs"));
 
 	if (FString const* ContentBundlesToConvertString = CommandLineParams.Find(TEXT("ContentBundles")))
 	{
@@ -240,6 +243,13 @@ bool UGameFeatureActionConvertContentBundleWorldPartitionBuilder::RunInternal(UW
 			PackagesToSave.Add(EDLWorldDataLayers->GetPackage());
 		}
 
+		// In case we want to delete HLODs
+		TArray<UClass*> HLODActorClasses =
+		{
+			AWorldPartitionHLOD::StaticClass(),
+			FindObject<UClass>(nullptr, TEXT("/Script/Engine.SpatialHashRuntimeGridInfo"))
+		};
+
 		TArray<FGuid> ConvertedActorGuids;
 		TArray<FWorldPartitionReference> ActorReferences;
 		bool bContentBundleActorConversionSuccess = true;
@@ -259,6 +269,20 @@ bool UGameFeatureActionConvertContentBundleWorldPartitionBuilder::RunInternal(UW
 			UPackage* OldActorPackage = Actor->GetExternalPackage();
 			ResetLoaders(OldActorPackage);
 
+			// Content bundles can contain HLODs, and since their names are partly based on the content bundle they were in, moving/renaming them over in an EDL is not straightforward.
+			// It's easier to delete them and let the automated nightly recreate them properly.
+			// Only need to use this option if a HLOD is preventing FExternalDataLayerHelper::MoveActorsToExternalDataLayer from succeeding, and you want to avoid a manual Remove HLOD step
+			if (bDeleteHLODs)
+			{
+				if (HLODActorClasses.FindByPredicate([ActorClass = Actor->GetClass()](const UClass* HLODClass) { return ActorClass->IsChildOf(HLODClass); }))
+				{
+					PackagesToDelete.Add(Actor->GetExternalPackage());
+					ConvertedActorGuids.Add(It->GetGuid());
+					UE_LOG(LogConvertContentBundleBuilder, Log, TEXT("Deleted HLOD-related Actor %s."), *Actor->GetFullName());
+					continue;
+				}
+			}
+
 			FText FailureReason;
 			if (!FExternalDataLayerHelper::MoveActorsToExternalDataLayer({ Actor }, ExternalDataLayerInstance, &FailureReason))
 			{
@@ -423,3 +447,4 @@ UExternalDataLayerAsset* UGameFeatureActionConvertContentBundleWorldPartitionBui
 	}
 	return NewExternalDataLayerAsset;
 }
+
diff --git a/Engine/Plugins/Runtime/GameFeatures/Source/GameFeaturesEditor/Private/GameFeatureActionConvertContentBundleWorldPartitionBuilder.h b/Engine/Plugins/Runtime/GameFeatures/Source/GameFeaturesEditor/Private/GameFeatureActionConvertContentBundleWorldPartitionBuilder.h
--- a/Engine/Plugins/Runtime/GameFeatures/Source/GameFeaturesEditor/Private/GameFeatureActionConvertContentBundleWorldPartitionBuilder.h
+++ b/Engine/Plugins/Runtime/GameFeatures/Source/GameFeaturesEditor/Private/GameFeatureActionConvertContentBundleWorldPartitionBuilder.h
@@ -33,4 +33,5 @@ private:
 	FString DestinationFolder;
 	bool bReportOnly;
 	bool bRemoveContentBundleAction;
+	bool bDeleteHLODs;
 };
\ No newline at end of file
