Remove duplicated copies of IsGameFeaturePlugin function
#rnx

#Synced-CL 50516840

diff --git a/Engine/Plugins/Runtime/GameFeatures/Source/GameFeaturesEditor/Private/IllegalPluginDependenciesValidator.cpp b/Engine/Plugins/Runtime/GameFeatures/Source/GameFeaturesEditor/Private/IllegalPluginDependenciesValidator.cpp
--- a/Engine/Plugins/Runtime/GameFeatures/Source/GameFeaturesEditor/Private/IllegalPluginDependenciesValidator.cpp
+++ b/Engine/Plugins/Runtime/GameFeatures/Source/GameFeaturesEditor/Private/IllegalPluginDependenciesValidator.cpp
@@ -10,25 +10,6 @@
 
 #define LOCTEXT_NAMESPACE "IllegalPluginDependenciesValidator"
 
-namespace IllegalPluginDependencies
-{
-	bool IsGameFeaturePlugin(const TSharedPtr<IPlugin> InPlugin)
-	{
-		const FString PluginsFolderRoot = (InPlugin->GetLoadedFrom() == EPluginLoadedFrom::Project) ? FPaths::ProjectPluginsDir() : FPaths::EnginePluginsDir();
-		FString PluginPathRelativeToDomain = InPlugin->GetBaseDir();
-		if (FPaths::MakePathRelativeTo(PluginPathRelativeToDomain, *PluginsFolderRoot))
-		{
-			PluginPathRelativeToDomain = TEXT("/") + PluginPathRelativeToDomain + TEXT("/");
-		}
-		else
-		{
-			PluginPathRelativeToDomain = FString();
-		}
-
-		return (PluginPathRelativeToDomain.StartsWith(TEXT("/GameFeatures/")));
-	}
-}
-
 bool UIllegalPluginDependenciesValidator::CanValidateAsset_Implementation(const FAssetData& AssetData, UObject* InAsset, FDataValidationContext& InContext) const
 {
 	return (InAsset->GetClass() == UDataValidationChangelist::StaticClass());
@@ -75,7 +56,7 @@ EDataValidationResult UIllegalPluginDependenciesValidator::ValidateLoadedAsset_I
 			TSharedPtr<IPlugin> Plugin = IPluginManager::Get().FindPlugin(ModifiedPlugin);
 			if (Plugin.IsValid())
 			{
-				if (GFPs.Contains(ModifiedPlugin) || IllegalPluginDependencies::IsGameFeaturePlugin(Plugin))
+				if (GFPs.Contains(ModifiedPlugin) || UGameFeaturesSubsystem::IsGameFeaturePlugin(Plugin))
 				{
 					continue;
 				}
