The AR singleton is null'ed out during GC at the same time our StateMachines can get GC'ed during PreAppExit. If the AR gets purged first from the GC we will crash trying to run these dtor.

Avoid calling the callback if we are exiting

diff --git a/Engine/Plugins/Runtime/GameFeatures/Source/GameFeatures/Private/GameFeaturesSubsystem.cpp b/Engine/Plugins/Runtime/GameFeatures/Source/GameFeatures/Private/GameFeaturesSubsystem.cpp
--- a/Engine/Plugins/Runtime/GameFeatures/Source/GameFeatures/Private/GameFeaturesSubsystem.cpp
+++ b/Engine/Plugins/Runtime/GameFeatures/Source/GameFeatures/Private/GameFeaturesSubsystem.cpp
@@ -4978,9 +4978,10 @@ UGameFeaturesSubsystem::FAmmortizedLoadContext::~FAmmortizedLoadContext()
 		}
 	}
 
-	CompleteDelegate.ExecuteIfBound(Results);
+	// if we are exiting we may have lost some of these singletones we are calling into, just avoid calling into to prevent random crashes on exit
 	if (!IsEngineExitRequested())
 	{
+		CompleteDelegate.ExecuteIfBound(Results);
 		UGameFeaturesSubsystem::Get().ActiveAmmoritzedLoads.Remove(GUID);
 	}
 }
diff --git a/Engine/Source/Runtime/AssetRegistry/Public/AssetRegistry/IAssetRegistry.h b/Engine/Source/Runtime/AssetRegistry/Public/AssetRegistry/IAssetRegistry.h
--- a/Engine/Source/Runtime/AssetRegistry/Public/AssetRegistry/IAssetRegistry.h
+++ b/Engine/Source/Runtime/AssetRegistry/Public/AssetRegistry/IAssetRegistry.h
@@ -1175,7 +1175,7 @@ public:
 	struct FSuspendPathsDismounting
 	{
 		UE_NODISCARD_CTOR FSuspendPathsDismounting() { IAssetRegistry::GetChecked().SuspendPathsDismounting(); }
-		~FSuspendPathsDismounting() { IAssetRegistry::GetChecked().ResumePathsDismounting(); }
+		~FSuspendPathsDismounting() { IAssetRegistry* AR = IAssetRegistry::Get(); if (AR) { AR->ResumePathsDismounting(); } }
 	};
 
 protected:
