Use FGameFeaturePluginIdentifier to find the PluginName from the PluginURL to avoid using the StateMachine in case its not created

#Synced-CL 50700071

diff --git a/Engine/Plugins/Runtime/GameFeatures/Source/GameFeatures/Private/GameFeaturesSubsystem.cpp b/Engine/Plugins/Runtime/GameFeatures/Source/GameFeatures/Private/GameFeaturesSubsystem.cpp
--- a/Engine/Plugins/Runtime/GameFeatures/Source/GameFeatures/Private/GameFeaturesSubsystem.cpp
+++ b/Engine/Plugins/Runtime/GameFeatures/Source/GameFeatures/Private/GameFeaturesSubsystem.cpp
@@ -2429,7 +2429,8 @@ void UGameFeaturesSubsystem::TerminateGameFeaturePlugin(TConstArrayView<FString>
 bool UGameFeaturesSubsystem::CancellAmortizedLoads(const FString& PluginURL)
 {
 	bool bDidCancellAnything = false;
-	const FString CancellingPluginName = GetPluginNameFromPluginURL(PluginURL);
+	const FGameFeaturePluginIdentifier CancellingPluginNameId(PluginURL);
+	const FString CancellingPluginName = FString(CancellingPluginNameId.GetPluginName());
 	if (!CancellingPluginName.IsEmpty())
 	{
 		for (TPair<FGuid, TWeakPtr<FAmmortizedLoadContext>>& AmmortizedLoad : ActiveAmmoritzedLoads)
@@ -2448,6 +2449,10 @@ bool UGameFeaturesSubsystem::CancellAmortizedLoads(const FString& PluginURL)
 			}
 		}
 	}
+	else
+	{
+		UE_LOG(LogGameFeatures, Error, TEXT("Failed to find the PluginName from the PluginURL(%s) failing to cancel this plugin"), *PluginURL);
+	}
 
 	return bDidCancellAnything;
 }
