Reduce work if calling AddOrUpdate on the same GFP + State already ref counted

#Synced-CL 50193382

diff --git a/Engine/Plugins/Runtime/GameFeatures/Source/GameFeatures/Private/GameFeatureStateHandle.cpp b/Engine/Plugins/Runtime/GameFeatures/Source/GameFeatures/Private/GameFeatureStateHandle.cpp
--- a/Engine/Plugins/Runtime/GameFeatures/Source/GameFeatures/Private/GameFeatureStateHandle.cpp
+++ b/Engine/Plugins/Runtime/GameFeatures/Source/GameFeatures/Private/GameFeatureStateHandle.cpp
@@ -113,6 +113,18 @@ void FGameFeatureStateHandleInternal::AddOrUpdateOwnershipIfHighestDestState(con
 	}
 }
 
+bool FGameFeatureStateHandleInternal::FindPluginRequiredState(const FString& PluginName, EGameFeaturePluginState& OutPluginState) const
+{
+	if (const EGameFeaturePluginState* FoundPluginState = StateHandleData.Find(PluginName))
+	{
+		OutPluginState = *FoundPluginState;
+		return true;
+	}
+
+	OutPluginState = EGameFeaturePluginState::UnknownStatus;
+	return false;
+}
+
 const TArray<FString> FGameFeatureStateHandleInternal::GetPlugins() const
 {
 	TArray<FString> OutPlugins;
diff --git a/Engine/Plugins/Runtime/GameFeatures/Source/GameFeatures/Private/GameFeatureStateHandleInternal.h b/Engine/Plugins/Runtime/GameFeatures/Source/GameFeatures/Private/GameFeatureStateHandleInternal.h
--- a/Engine/Plugins/Runtime/GameFeatures/Source/GameFeatures/Private/GameFeatureStateHandleInternal.h
+++ b/Engine/Plugins/Runtime/GameFeatures/Source/GameFeatures/Private/GameFeatureStateHandleInternal.h
@@ -26,6 +26,7 @@ public:
 	void TakeOwnership(FGameFeatureStateHandleInternal& StateHandle);
 
 	void AddOrUpdateOwnershipIfHighestDestState(const FString& PluginName, EGameFeaturePluginState DestPluginState);
+	bool FindPluginRequiredState(const FString& PluginName, EGameFeaturePluginState& OutPluginState) const;
 
 	const TArray<FString> GetPlugins() const;
 	int GetPluginCount() const;
diff --git a/Engine/Plugins/Runtime/GameFeatures/Source/GameFeatures/Private/GameFeatureStateHandleReferenceController.cpp b/Engine/Plugins/Runtime/GameFeatures/Source/GameFeatures/Private/GameFeatureStateHandleReferenceController.cpp
--- a/Engine/Plugins/Runtime/GameFeatures/Source/GameFeatures/Private/GameFeatureStateHandleReferenceController.cpp
+++ b/Engine/Plugins/Runtime/GameFeatures/Source/GameFeatures/Private/GameFeatureStateHandleReferenceController.cpp
@@ -300,6 +300,13 @@ void FGameFeatureStateHandleReferenceController::AddOrUpdateReference(const FGui
 		return;
 	}
 
+	FString PluginName = GetPluginNameFromURL(GameFeaturePluginURL);
+	EGameFeaturePluginState CurrentPluginState;
+	if (InternalHandle->FindPluginRequiredState(PluginName, CurrentPluginState) && CurrentPluginState == PluginState) // if we are already at this state done recalc
+	{
+		return;
+	}
+
 	UGameFeaturesSubsystem& GameFeatures = UGameFeaturesSubsystem::Get();
 
 	TMap<FString, EGameFeaturePluginState> AllGFPDependencies;
@@ -312,18 +319,18 @@ void FGameFeatureStateHandleReferenceController::AddOrUpdateReference(const FGui
 		AllGFPDependencies.Add(GameFeaturePluginURL, PluginState);
 	}
 
-	UE_LOG(LogGameFeatureStateHandle, Verbose, TEXT("Adding reference to %s (%s) and depends. Self + Depends count: %i"), *GetPluginNameFromURL(GameFeaturePluginURL), *UE::GameFeatures::ToString(PluginState), AllGFPDependencies.Num());
+	UE_LOG(LogGameFeatureStateHandle, Verbose, TEXT("Adding reference to %s (%s) and depends. Self + Depends count: %i"), *PluginName, *UE::GameFeatures::ToString(PluginState), AllGFPDependencies.Num());
 	for (const TPair<FString, EGameFeaturePluginState>& Depends : AllGFPDependencies)
 	{
-		FString PluginName = GetPluginNameFromURL(Depends.Key);
+		FString DependPluginName = GetPluginNameFromURL(Depends.Key);
 		// Update our GFP Name -> RefCount mapping
-		FGameFeatureStateRefCount& DependsRefCount = Internal->GFPReferenceCount.FindOrAdd(PluginName);
+		FGameFeatureStateRefCount& DependsRefCount = Internal->GFPReferenceCount.FindOrAdd(DependPluginName);
 		DependsRefCount.AddRef(StateHandleId, Depends.Value);
 
 		// Add or Update our InternalHandle on the Depends it requires at a min state
-		InternalHandle->AddOrUpdateOwnershipIfHighestDestState(PluginName, Depends.Value);
+		InternalHandle->AddOrUpdateOwnershipIfHighestDestState(DependPluginName, Depends.Value);
 
-		UE_LOG(LogGameFeatureStateHandle, Verbose, TEXT("  depends %s (%s)"), *PluginName, *UE::GameFeatures::ToString(Depends.Value));
+		UE_LOG(LogGameFeatureStateHandle, Verbose, TEXT("  depends %s (%s)"), *DependPluginName, *UE::GameFeatures::ToString(Depends.Value));
 	}
 }
 
@@ -741,7 +748,7 @@ void FGameFeatureStateHandleReferenceControllerInternal::ListGameFeatureStateHan
 		{
 			if (FGameFeatureStateHandleInternal* InternalHandle = StateHandles.Find(OwnerToState.Key))
 			{
-				RefCountStr.Appendf(TEXT("(%s (%s)"), *InternalHandle->GetOwner(), *UE::GameFeatures::ToString(OwnerToState.Value));
+				RefCountStr.Appendf(TEXT("(%s (%s))"), *InternalHandle->GetOwner(), *UE::GameFeatures::ToString(OwnerToState.Value));
 			}
 		}
 		RefCountStr.Append(TEXT("}"));
