Add  CTOR that does Init/Register if desired

#Synced-CL 50193382

diff --git a/Engine/Plugins/Runtime/GameFeatures/Source/GameFeatures/Private/GameFeatureStateHandle.cpp b/Engine/Plugins/Runtime/GameFeatures/Source/GameFeatures/Private/GameFeatureStateHandle.cpp
--- a/Engine/Plugins/Runtime/GameFeatures/Source/GameFeatures/Private/GameFeatureStateHandle.cpp
+++ b/Engine/Plugins/Runtime/GameFeatures/Source/GameFeatures/Private/GameFeatureStateHandle.cpp
@@ -5,6 +5,11 @@
 #include "GameFeatureStateHandleReferenceController.h"
 #include "GameFeaturesSubsystem.h"
 
+FGameFeatureStateHandle::FGameFeatureStateHandle(const FString& Owner, EGameFeatureStateHandleOptions Options)
+{
+	InitAndRegister(Owner, Options);
+}
+
 FGameFeatureStateHandle::~FGameFeatureStateHandle()
 {
 	ReleaseAndUnregister();
diff --git a/Engine/Plugins/Runtime/GameFeatures/Source/GameFeatures/Public/GameFeatureStateHandle.h b/Engine/Plugins/Runtime/GameFeatures/Source/GameFeatures/Public/GameFeatureStateHandle.h
--- a/Engine/Plugins/Runtime/GameFeatures/Source/GameFeatures/Public/GameFeatureStateHandle.h
+++ b/Engine/Plugins/Runtime/GameFeatures/Source/GameFeatures/Public/GameFeatureStateHandle.h
@@ -27,6 +27,7 @@ class FGameFeatureStateHandle
 {
 public:
 	FGameFeatureStateHandle() = default;
+	explicit UE_API FGameFeatureStateHandle(const FString& Owner, EGameFeatureStateHandleOptions Options = EGameFeatureStateHandleOptions::TrackDependencies);
 
 	/** Releases all the ref counts, and if any GFPs need to be downgraded moves them. When done we call the Callback if setup. The Dtor call this */
 	UE_API void ReleaseAndUnregister(TFunction<void(bool)> ResetCompleteCallback = [](bool){});
