[GFP] Add ability to disable GC batching when transitioning GFPs to lowest state

#rnx

diff --git a/Engine/Plugins/Runtime/GameFeatures/Source/GameFeatures/Private/GameFeaturesSubsystem.cpp b/Engine/Plugins/Runtime/GameFeatures/Source/GameFeatures/Private/GameFeaturesSubsystem.cpp
--- a/Engine/Plugins/Runtime/GameFeatures/Source/GameFeatures/Private/GameFeaturesSubsystem.cpp
+++ b/Engine/Plugins/Runtime/GameFeatures/Source/GameFeatures/Private/GameFeaturesSubsystem.cpp
@@ -130,6 +130,18 @@ namespace UE::GameFeatures
 		TEXT("attributed to the first one that actually causes the load. Stats can be output with ")
 		TEXT("GameFeaturePlugin.PrintLoadStats."));
 
+	static FAutoConsoleVariable CVarDisablePluginUnloadBatchingDuringTermination(
+		TEXT("GameFeaturePlugin.DisablePluginUnloadBatchingDuringTermination"),
+		false,
+		TEXT("Development only: Disables the plugin manager unload batchinig, ")
+		TEXT("so every unload of GFP plugin will trigger Garbage collection. Happens during batch Termination."));
+
+	static FAutoConsoleVariable CVarDisablePluginUnloadBatchingDuringUnload(
+		TEXT("GameFeaturePlugin.DisablePluginUnloadBatchingDuringUnload"),
+		false,
+		TEXT("Development only: Disables the plugin manager unload batchinig, ")
+		TEXT("so every unload of GFP plugin will trigger Garbage collection. Happens during batch Unload."));
+
 	TOptional<FString> GetPluginUrlForConsoleCommand(const TArray<FString>& Args, FOutputDevice& Ar)
 	{
 		TOptional<FString> PluginURL;
@@ -2252,6 +2264,10 @@ void UGameFeaturesSubsystem::UnloadGameFeaturePlugin(TConstArrayView<FString> Pl
 	{
 		TransitionOptions = GameFeaturesSubsystem::ETransitionOptions::All;
 	}
+	if (UE::GameFeatures::CVarDisablePluginUnloadBatchingDuringUnload->GetBool())
+	{
+		EnumRemoveFlags(TransitionOptions, GameFeaturesSubsystem::ETransitionOptions::SuspendPluginManager);
+	}
 
 	TSharedRef<GameFeaturesSubsystem::FTransitionContenxt> TransitionContext = MakeShared<GameFeaturesSubsystem::FTransitionContenxt>(TransitionOptions);
 	TransitionContext->CompleteDelegate = CompleteDelegate;
@@ -2404,7 +2420,13 @@ void UGameFeaturesSubsystem::TerminateGameFeaturePlugin(const FString& PluginURL
 
 void UGameFeaturesSubsystem::TerminateGameFeaturePlugin(TConstArrayView<FString> PluginURLs, const FMultipleGameFeaturePluginsTerminated& CompleteDelegate)
 {
-	TSharedRef<GameFeaturesSubsystem::FTransitionContenxt> TransitionContext = MakeShared<GameFeaturesSubsystem::FTransitionContenxt>();
+	GameFeaturesSubsystem::ETransitionOptions TransitionOptions = GameFeaturesSubsystem::ETransitionOptions::All;
+	if (UE::GameFeatures::CVarDisablePluginUnloadBatchingDuringTermination->GetBool())
+	{
+		EnumRemoveFlags(TransitionOptions, GameFeaturesSubsystem::ETransitionOptions::SuspendPluginManager);
+	}
+
+	TSharedRef<GameFeaturesSubsystem::FTransitionContenxt> TransitionContext = MakeShared<GameFeaturesSubsystem::FTransitionContenxt>(TransitionOptions);
 	TransitionContext->CompleteDelegate = CompleteDelegate;
 
 	TransitionContext->Results.Reserve(PluginURLs.Num());
