Remove ifdef'd out profiling code

#Synced-CL 50901530

diff --git a/Engine/Plugins/Runtime/GameFeatures/Source/GameFeatures/Private/GameFeaturePluginStateMachine.cpp b/Engine/Plugins/Runtime/GameFeatures/Source/GameFeatures/Private/GameFeaturePluginStateMachine.cpp
--- a/Engine/Plugins/Runtime/GameFeatures/Source/GameFeatures/Private/GameFeaturePluginStateMachine.cpp
+++ b/Engine/Plugins/Runtime/GameFeatures/Source/GameFeatures/Private/GameFeaturePluginStateMachine.cpp
@@ -4303,7 +4303,6 @@ bool UE::GameFeatures::FStateMachineExecutor::UpdateStateMachines(int32 MaxItera
 	return OrderedUpdateQueue.Num() > 0;
 }
 
-#if 1
 #define GFP_STATE_OP(STATE_NAME, TXT) \
 template<> \
 void UE::GameFeatures::FStateMachineExecutor::UpdateStateMachinesGrouped<EGameFeaturePluginState::STATE_NAME>(TConstArrayView<UGameFeaturePluginStateMachine*> StateMachines) \
@@ -4329,22 +4328,6 @@ void UE::GameFeatures::FStateMachineExecutor::UpdateStateMachinesGrouped(EGameFe
 #undef GFP_STATE_OP
 	}
 }
-#else
-UE_TRACE_EVENT_BEGIN(Cpu, GFP_UpdateStateMachinesGrouped, NoSync)
-	UE_TRACE_EVENT_FIELD(UE::Trace::WideString, State)
-UE_TRACE_EVENT_END()
-
-void UE::GameFeatures::FStateMachineExecutor::UpdateStateMachinesGrouped(EGameFeaturePluginState State, TConstArrayView<UGameFeaturePluginStateMachine*> StateMachines)
-{
-	UE_TRACE_LOG_SCOPED_T(Cpu, GFP_UpdateStateMachinesGrouped, CpuChannel)
-		<< GFP_UpdateStateMachinesGrouped.State(*UE::GameFeatures::ToString(State));
-
-	for (UGameFeaturePluginStateMachine* StateMachine : StateMachines)
-	{
-		OwningSubsystem->UpdateStateMachine(StateMachine, { 1 });
-	}
-}
-#endif
 
 void UE::GameFeatures::FStateMachineExecutor::QueueUpdateStateMachine(FGameFeaturePluginStateMachineProperties& InProperties)
 {
