Add Amortized load built in game feature plugins call to mirror the GameFeaturesSubsystem.

diff --git a/Engine/Plugins/Runtime/GameFeatures/Source/GameFeatures/Private/GameFeatureStateHandleReferenceController.cpp b/Engine/Plugins/Runtime/GameFeatures/Source/GameFeatures/Private/GameFeatureStateHandleReferenceController.cpp
--- a/Engine/Plugins/Runtime/GameFeatures/Source/GameFeatures/Private/GameFeatureStateHandleReferenceController.cpp
+++ b/Engine/Plugins/Runtime/GameFeatures/Source/GameFeatures/Private/GameFeatureStateHandleReferenceController.cpp
@@ -676,6 +676,31 @@ void FGameFeatureStateHandleReferenceController::RegisterGameFeaturePlugin(const
 	}));
 }
 
+void FGameFeatureStateHandleReferenceController::LoadBuiltInGameFeaturePlugins_Amortized(const FGameFeatureStateHandle& StateHandle, const TArray<TSharedRef<IPlugin>>& Plugins, const UGameFeaturesSubsystem::FBuiltInPluginAdditionalFilters_Copyable& AdditionalFilter, int32 AmortizeRateMS, const UMultipleGameFeatureStateHandleLoadComplete& CompleteDelegate)
+{
+	UGameFeaturesSubsystem::Get().LoadBuiltInGameFeaturePlugins_Amortized(Plugins, AdditionalFilter, AmortizeRateMS, FBuiltInGameFeaturePluginsLoaded::CreateLambda([this, CompleteDelegate, StateHandleId = StateHandle.GetUniqueId()](const TMap<FString, UE::GameFeatures::FResult>& Results)
+		{
+			if (IsLifetimeControlEnabled())
+			{
+				UGameFeaturesSubsystem& GameFeatureSubsystem = UGameFeaturesSubsystem::Get();
+				for (const TPair<FString, UE::GameFeatures::FResult>& Result : Results)
+				{
+					FString PluginURL;
+					if (GameFeatureSubsystem.GetPluginURLByName(Result.Key, PluginURL))
+					{
+						AddOrUpdateReference(StateHandleId, PluginURL, GameFeatureSubsystem.GetPluginState(PluginURL));
+					}
+					else
+					{
+						UE_LOG(LogGameFeatureStateHandle, Error, TEXT("Failed to add reference to plugin '%s'. Unable to find matching PluginURL."), *Result.Key);
+					}
+				}
+			}
+
+			CompleteDelegate.ExecuteIfBound(Results);
+		}));
+}
+
 void FGameFeatureStateHandleReferenceController::ResetGameFeatureStateHandle(const FGameFeatureStateHandle& StateHandle, TFunction<void(bool)> OnComplete)
 {
 	if (!IsLifetimeControlEnabled())
diff --git a/Engine/Plugins/Runtime/GameFeatures/Source/GameFeatures/Public/GameFeatureStateHandleReferenceController.h b/Engine/Plugins/Runtime/GameFeatures/Source/GameFeatures/Public/GameFeatureStateHandleReferenceController.h
--- a/Engine/Plugins/Runtime/GameFeatures/Source/GameFeatures/Public/GameFeatureStateHandleReferenceController.h
+++ b/Engine/Plugins/Runtime/GameFeatures/Source/GameFeatures/Public/GameFeatureStateHandleReferenceController.h
@@ -6,9 +6,11 @@
 #include "GameFeatureTypes.h"
 #include "Delegates/Delegate.h"
 #include "Delegates/DelegateCombinations.h"
+#include "GameFeaturesSubsystem.h"
 
 namespace UE::GameFeatures { struct FResult; }
 struct FGameFeatureProtocolOptions;
+class IPlugin;
 
 #define UE_API GAMEFEATURES_API
 
@@ -35,6 +37,8 @@ public:
 	UE_API void RegisterGameFeaturePlugin(const FGameFeatureStateHandle& StateHandle, const FString& PluginURL, const FGameFeatureProtocolOptions& ProtocolOptions, const UGameFeatureStateHandleLoadComplete& CompleteDelegate);
 	UE_API void RegisterGameFeaturePlugin(const FGameFeatureStateHandle& StateHandle, const TArray<FString>& PluginURLs, const FGameFeatureProtocolOptions& ProtocolOptions, const UMultipleGameFeatureStateHandleLoadComplete& CompleteDelegate);
 
+	UE_API void LoadBuiltInGameFeaturePlugins_Amortized(const FGameFeatureStateHandle& StateHandle, const TArray<TSharedRef<IPlugin>>& Plugins, const UGameFeaturesSubsystem::FBuiltInPluginAdditionalFilters_Copyable& AdditionalFilter, int32 AmortizeRateMS, const UMultipleGameFeatureStateHandleLoadComplete& CompleteDelegate);
+
 	/** Removes all the GFPs references held by this list of handles, and will empty out the StateHandle list */
 	UE_API void ResetGameFeatureStateHandle(const FGameFeatureStateHandle& StateHandle, TFunction<void(bool)> OnComplete = [](bool) {});
 
