Add FStringView overload for IsValidGameFeaturePlugin

#Synced-CL 50925863

diff --git a/Engine/Plugins/Runtime/GameFeatures/Source/GameFeatures/Private/GameFeaturesSubsystemSettings.cpp b/Engine/Plugins/Runtime/GameFeatures/Source/GameFeatures/Private/GameFeaturesSubsystemSettings.cpp
--- a/Engine/Plugins/Runtime/GameFeatures/Source/GameFeatures/Private/GameFeaturesSubsystemSettings.cpp
+++ b/Engine/Plugins/Runtime/GameFeatures/Source/GameFeatures/Private/GameFeaturesSubsystemSettings.cpp
@@ -17,6 +17,11 @@ UGameFeaturesSubsystemSettings::UGameFeaturesSubsystemSettings()
 }
 
 bool UGameFeaturesSubsystemSettings::IsValidGameFeaturePlugin(const FString& PluginDescriptorFilename) const
+{
+	return IsValidGameFeaturePlugin(FStringView(PluginDescriptorFilename));
+}
+
+bool UGameFeaturesSubsystemSettings::IsValidGameFeaturePlugin(FStringView PluginDescriptorFilename) const
 {
 	// Build the cache of game feature plugin folders the first time this is called
 	static struct FBuiltInGameFeaturePluginsFolders
diff --git a/Engine/Plugins/Runtime/GameFeatures/Source/GameFeatures/Public/GameFeaturesSubsystemSettings.h b/Engine/Plugins/Runtime/GameFeatures/Source/GameFeatures/Public/GameFeaturesSubsystemSettings.h
--- a/Engine/Plugins/Runtime/GameFeatures/Source/GameFeatures/Public/GameFeaturesSubsystemSettings.h
+++ b/Engine/Plugins/Runtime/GameFeatures/Source/GameFeatures/Public/GameFeaturesSubsystemSettings.h
@@ -42,8 +42,12 @@ public:
 	FString BuiltInGameFeaturePluginsFolder;
 
 public:
-	// Returns true if the specified (normalized or full) path is a game feature plugin
+	/** 
+	  * Returns true if the specified (normalized or full) path is a game feature plugin
+	  * Safe to call from any thread.
+	  */ 
 	UE_API bool IsValidGameFeaturePlugin(const FString& PluginDescriptorFilename) const;
+	UE_API bool IsValidGameFeaturePlugin(FStringView PluginDescriptorFilename) const;
 
 };
 
