[GFP] Suspend asset streamer in GameFeaturesSubsystem::FTransitionContext to prevent potential GC race while render asset streamer is running (disabled by default)

#Synced-CL 50925182

diff --git a/Engine/Plugins/Runtime/GameFeatures/Source/GameFeatures/Private/GameFeaturesSubsystem.cpp b/Engine/Plugins/Runtime/GameFeatures/Source/GameFeatures/Private/GameFeaturesSubsystem.cpp
--- a/Engine/Plugins/Runtime/GameFeatures/Source/GameFeatures/Private/GameFeaturesSubsystem.cpp
+++ b/Engine/Plugins/Runtime/GameFeatures/Source/GameFeatures/Private/GameFeaturesSubsystem.cpp
@@ -35,6 +35,7 @@
 #include "Engine/StreamableManager.h"
 #include "Engine/AssetManagerSettings.h"
 #include "InstallBundleTypes.h"
+#include "RenderAssetUpdate.h"
 
 #include UE_INLINE_GENERATED_CPP_BY_NAME(GameFeaturesSubsystem)
 
@@ -142,6 +143,11 @@ namespace UE::GameFeatures
 		TEXT("Development only: Disables the plugin manager unload batchinig, ")
 		TEXT("so every unload of GFP plugin will trigger Garbage collection. Happens during batch Unload."));
 
+	static FAutoConsoleVariable CVarTransitionContextSuspendStreamingManager(
+		TEXT("GameFeaturePlugin.TransitionContextSuspendStreamingManager"),
+		false,
+		TEXT("Suspend streaming manager in transition context to prevent potential GC races."));
+
 	TOptional<FString> GetPluginUrlForConsoleCommand(const TArray<FString>& Args, FOutputDevice& Ar)
 	{
 		TOptional<FString> PluginURL;
@@ -2093,14 +2099,27 @@ namespace GameFeaturesSubsystem
 		}
 	};
 
-	struct FSuspendPluginManger
+	struct FSuspendRenderAssetStreamer
+	{
+		UE_NODISCARD_CTOR FSuspendRenderAssetStreamer()
+		{
+			SuspendRenderAssetStreaming();
+		}
+
+		~FSuspendRenderAssetStreamer()
+		{
+			ResumeRenderAssetStreaming();
+		}
+	};
+
+	struct FSuspendPluginManager
 	{
-		UE_NODISCARD_CTOR FSuspendPluginManger()
+		UE_NODISCARD_CTOR FSuspendPluginManager()
 		{
 			IPluginManager::Get().SuppressPluginUnloadGC();
 		}
 
-		~FSuspendPluginManger()
+		~FSuspendPluginManager()
 		{
 			IPluginManager::Get().ResumePluginUnloadGC();
 		}
@@ -2138,7 +2157,8 @@ namespace GameFeaturesSubsystem
 		SuspendGameplayTagTreeBroadcast	= 1 << 1,
 		SuspendAssetRegistry			= 1 << 2,
 		NotifyObservers					= 1 << 3,
-		All = SuspendPluginManager | SuspendGameplayTagTreeBroadcast | SuspendAssetRegistry | NotifyObservers,
+		SuspendRenderAssetStreamer		= 1 << 4,
+		All = SuspendPluginManager | SuspendGameplayTagTreeBroadcast | SuspendAssetRegistry | NotifyObservers | SuspendRenderAssetStreamer,
 	};
 	ENUM_CLASS_FLAGS(ETransitionOptions)
 
@@ -2148,13 +2168,22 @@ namespace GameFeaturesSubsystem
 		FMultipleGameFeaturePluginsReleased CompleteDelegate;
 		int32 NumPluginsTransitioned = 0;
 
-		TOptional<FSuspendPluginManger> SuspendPluginManager;
+		TOptional<FSuspendRenderAssetStreamer> SuspendRenderAssetStreamer; 
+		TOptional<FSuspendPluginManager> SuspendPluginManager;
 		TOptional<FSuspendGameplayTagTreeChangedBroadcast> SuspendGameplayTagTreeBroadcast;
 		TOptional<IAssetRegistry::FSuspendPathsDismounting> SuspendAssetRegistry;
 		TOptional<FNotifyObservers> NotifyObservers;
 
 		UE_NODISCARD_CTOR FTransitionContenxt(ETransitionOptions TransitionOptions = ETransitionOptions::All)
 		{
+			if (EnumHasAnyFlags(TransitionOptions, ETransitionOptions::SuspendRenderAssetStreamer))
+			{
+				if (!IsEngineExitRequested())
+				{
+					SuspendRenderAssetStreamer.Emplace();
+				}
+			}
+
 			if (EnumHasAnyFlags(TransitionOptions, ETransitionOptions::SuspendPluginManager))
 			{
 				if (!IsEngineExitRequested())
@@ -2183,7 +2212,7 @@ namespace GameFeaturesSubsystem
 		{
 			CompleteDelegate.ExecuteIfBound(Results);
 			// Explicit order of destruction is required
-			// At the very least we need to execute PluginManger batch last
+			// At the very least we need to execute PluginManger batch at the end and render asset streamer last
 			// as it will run GC and does the leaked packages verification.
 			// And by that point all other resources need to be properly cleared,
 			// i.e. other batches need to finish.
@@ -2191,6 +2220,7 @@ namespace GameFeaturesSubsystem
 			SuspendAssetRegistry.Reset();
 			SuspendGameplayTagTreeBroadcast.Reset();
 			SuspendPluginManager.Reset();
+			SuspendRenderAssetStreamer.Reset();
 		}
 	};
 }
@@ -2268,6 +2298,10 @@ void UGameFeaturesSubsystem::UnloadGameFeaturePlugin(TConstArrayView<FString> Pl
 	{
 		EnumRemoveFlags(TransitionOptions, GameFeaturesSubsystem::ETransitionOptions::SuspendPluginManager);
 	}
+	if (!UE::GameFeatures::CVarTransitionContextSuspendStreamingManager->GetBool())
+	{
+		EnumRemoveFlags(TransitionOptions, GameFeaturesSubsystem::ETransitionOptions::SuspendRenderAssetStreamer);
+	}
 
 	TSharedRef<GameFeaturesSubsystem::FTransitionContenxt> TransitionContext = MakeShared<GameFeaturesSubsystem::FTransitionContenxt>(TransitionOptions);
 	TransitionContext->CompleteDelegate = CompleteDelegate;
@@ -2425,6 +2459,10 @@ void UGameFeaturesSubsystem::TerminateGameFeaturePlugin(TConstArrayView<FString>
 	{
 		EnumRemoveFlags(TransitionOptions, GameFeaturesSubsystem::ETransitionOptions::SuspendPluginManager);
 	}
+	if (!UE::GameFeatures::CVarTransitionContextSuspendStreamingManager->GetBool())
+	{
+		EnumRemoveFlags(TransitionOptions, GameFeaturesSubsystem::ETransitionOptions::SuspendRenderAssetStreamer);
+	}
 
 	TSharedRef<GameFeaturesSubsystem::FTransitionContenxt> TransitionContext = MakeShared<GameFeaturesSubsystem::FTransitionContenxt>(TransitionOptions);
 	TransitionContext->CompleteDelegate = CompleteDelegate;
diff --git a/Engine/Source/Runtime/Engine/Public/RenderAssetUpdate.h b/Engine/Source/Runtime/Engine/Public/RenderAssetUpdate.h
--- a/Engine/Source/Runtime/Engine/Public/RenderAssetUpdate.h
+++ b/Engine/Source/Runtime/Engine/Public/RenderAssetUpdate.h
@@ -296,7 +296,7 @@ protected:
 	}
 };
 
-void SuspendRenderAssetStreaming();
-void ResumeRenderAssetStreaming();
+ENGINE_API void SuspendRenderAssetStreaming();
+ENGINE_API void ResumeRenderAssetStreaming();
 
 #undef UE_API
