Add IrisFilterGameFeatureAction to Engine

Provides dynamic filter creation/destruction for game feature plugins.

[RN] Added IrisFilterGameFeatureAction for game feature plugins to create & destroy network filters at runtime.

#Synced-CL 50901758

diff --git a/Engine/Plugins/Runtime/GameFeatures/Source/GameFeatures/GameFeatures.Build.cs b/Engine/Plugins/Runtime/GameFeatures/Source/GameFeatures/GameFeatures.Build.cs
--- a/Engine/Plugins/Runtime/GameFeatures/Source/GameFeatures/GameFeatures.Build.cs
+++ b/Engine/Plugins/Runtime/GameFeatures/Source/GameFeatures/GameFeatures.Build.cs
@@ -27,6 +27,7 @@ namespace UnrealBuildTool.Rules
 					"GameplayTags",
 					"InstallBundleManager",
 					"IoStoreOnDemandCore",
+					"IrisCore",
 					"Json",
 					"JsonUtilities",
 					"PakFile",
diff --git a/Engine/Plugins/Runtime/GameFeatures/Source/GameFeatures/Private/IrisFilterGameFeatureAction.cpp b/Engine/Plugins/Runtime/GameFeatures/Source/GameFeatures/Private/IrisFilterGameFeatureAction.cpp
new file mode 100644
--- /dev/null
+++ b/Engine/Plugins/Runtime/GameFeatures/Source/GameFeatures/Private/IrisFilterGameFeatureAction.cpp
@@ -0,0 +1,120 @@
+// Copyright Epic Games, Inc. All Rights Reserved.
+
+#include "IrisFilterGameFeatureAction.h"
+
+#include "GameFeaturesSubsystem.h"
+#include "Iris/ReplicationSystem/ObjectReplicationBridge.h"
+#include "Iris/ReplicationSystem/ReplicationSystem.h"
+#include "Misc/ScopeExit.h"
+#include "Net/Iris/ReplicationSystem/ReplicationSystemUtil.h"
+
+#include UE_INLINE_GENERATED_CPP_BY_NAME(IrisFilterGameFeatureAction)
+
+DEFINE_LOG_CATEGORY_STATIC(LogIrisFilterGameFeatureAction, Log, Verbose)
+
+void UIrisFilterGameFeatureAction::OnGameFeatureActivating(FGameFeatureActivatingContext& ChangeContext)
+{
+	Super::OnGameFeatureActivating(ChangeContext);
+
+#if WITH_SERVER_CODE
+	for (const FWorldContext& WorldContext : GEngine->GetWorldContexts())
+	{
+		if (ShouldApplyToWorld(ChangeContext, WorldContext))
+		{
+			AddFilterToWorld(WorldContext.World());
+		}
+	}
+#endif // WITH_SERVER_CODE
+}
+
+void UIrisFilterGameFeatureAction::OnGameFeatureDeactivating(FGameFeatureDeactivatingContext& ChangeContext)
+{
+	ON_SCOPE_EXIT{ Super::OnGameFeatureDeactivating(ChangeContext); };
+
+#if WITH_SERVER_CODE
+	for (const FWorldContext& WorldContext : GEngine->GetWorldContexts())
+	{
+		if (ShouldApplyToWorld(ChangeContext, WorldContext))
+		{
+			RemoveFilterFromWorld(WorldContext.World());
+		}
+	}
+#endif // WITH_SERVER_CODE
+}
+
+bool UIrisFilterGameFeatureAction::ShouldApplyToWorld(const FGameFeatureStateChangeContext& ChangeContext, const FWorldContext& WorldContext) const
+{
+#if WITH_SERVER_CODE
+	if (!ChangeContext.ShouldApplyToWorldContext(WorldContext))
+	{
+		return false;
+	}
+
+	UWorld* World = WorldContext.World();
+	if (!World)
+	{
+		return false;
+	}
+
+	return World->GetNetMode() < NM_Client;
+#else // !WITH_SERVER_CODE
+	return false;
+#endif // !WITH_SERVER_CODE
+}
+
+void UIrisFilterGameFeatureAction::AddFilterToWorld(UWorld* World)
+{
+#if WITH_SERVER_CODE
+	if (!World)
+	{
+		return;
+	}
+
+	UReplicationSystem* ReplicationSystem = UE::Net::FReplicationSystemUtil::GetReplicationSystem(World);
+	if (!ReplicationSystem)
+	{
+		return;
+	}
+
+	if (!ReplicationSystem->CreateFilter(FilterDefinition).IsValid())
+	{
+		UE_LOG(LogIrisFilterGameFeatureAction, Warning, TEXT("Failed to create filter %s"), *FilterDefinition.FilterName.ToString());
+		return;
+	}
+
+	UE_LOG(
+		LogIrisFilterGameFeatureAction,
+		Verbose,
+		TEXT("Created filter %s in world %s"),
+		*FilterDefinition.FilterName.ToString(),
+		*World->GetName());
+
+	AddedToWorlds.Emplace(World);
+
+	if (UObjectReplicationBridge* ReplicationBridge = ReplicationSystem->GetReplicationBridge())
+	{
+		// Reload config to apply filter
+		ReplicationBridge->ReloadConfig();
+	}
+#endif // WITH_SERVER_CODE
+}
+
+void UIrisFilterGameFeatureAction::RemoveFilterFromWorld(UWorld* World)
+{
+#if WITH_SERVER_CODE
+	if (AddedToWorlds.Contains(World))
+	{
+		if (UReplicationSystem* ReplicationSystem = UE::Net::FReplicationSystemUtil::GetReplicationSystem(World))
+		{
+			ReplicationSystem->DestroyFilter(FilterDefinition.FilterName);
+
+			if (UObjectReplicationBridge* ReplicationBridge = ReplicationSystem->GetReplicationBridge())
+			{
+				ReplicationBridge->ReloadConfig();
+			}
+		}
+
+		AddedToWorlds.Remove(World);
+	}
+#endif // WITH_SERVER_CODE
+}
diff --git a/Engine/Plugins/Runtime/GameFeatures/Source/GameFeatures/Public/IrisFilterGameFeatureAction.h b/Engine/Plugins/Runtime/GameFeatures/Source/GameFeatures/Public/IrisFilterGameFeatureAction.h
new file mode 100644
--- /dev/null
+++ b/Engine/Plugins/Runtime/GameFeatures/Source/GameFeatures/Public/IrisFilterGameFeatureAction.h
@@ -0,0 +1,36 @@
+// Copyright Epic Games, Inc. All Rights Reserved.
+
+#pragma once
+
+#include "Containers/Set.h"
+#include "GameFeatureAction.h"
+#include "Iris/ReplicationSystem/Filtering/NetObjectFilterDefinitions.h"
+#include "UObject/WeakObjectPtrTemplates.h"
+
+#include "IrisFilterGameFeatureAction.generated.h"
+
+class UWorld;
+struct FGameFeatureStateChangeContext;
+struct FWorldContext;
+
+UCLASS(MinimalAPI, CollapseCategories, meta = (DisplayName = "Create Iris Network Filter"))
+class UIrisFilterGameFeatureAction : public UGameFeatureAction
+{
+	GENERATED_BODY()
+
+public:
+	virtual void OnGameFeatureActivating(FGameFeatureActivatingContext&) override;
+	virtual void OnGameFeatureDeactivating(FGameFeatureDeactivatingContext&) override;
+
+private:
+	bool ShouldApplyToWorld(const FGameFeatureStateChangeContext& StateChangeContext, const FWorldContext& WorldContext) const;
+	void AddFilterToWorld(UWorld* World);
+	void RemoveFilterFromWorld(UWorld* World);
+
+	// Filters to create on activation & destroy on deactivation.
+	UPROPERTY(EditAnywhere, Category=Filter)
+	FNetObjectFilterDefinition FilterDefinition;
+
+	// Worlds the filter was added to, and should be removed from.
+	TSet<TWeakObjectPtr<UWorld>> AddedToWorlds;
+};
