[Backout] - CL51187212
#rnx
Original CL Desc
-----------------------------------------------------------------
Add IrisFilterGameFeatureAction to Engine

Provides dynamic filter creation/destruction for game feature plugins.

[RN] Added IrisFilterGameFeatureAction for game feature plugins to create & destroy network filters at runtime.

diff --git a/Engine/Plugins/Runtime/GameFeatures/Source/GameFeatures/GameFeatures.Build.cs b/Engine/Plugins/Runtime/GameFeatures/Source/GameFeatures/GameFeatures.Build.cs
--- a/Engine/Plugins/Runtime/GameFeatures/Source/GameFeatures/GameFeatures.Build.cs
+++ b/Engine/Plugins/Runtime/GameFeatures/Source/GameFeatures/GameFeatures.Build.cs
@@ -27,7 +27,6 @@ namespace UnrealBuildTool.Rules
 					"GameplayTags",
 					"InstallBundleManager",
 					"IoStoreOnDemandCore",
-					"IrisCore",
 					"Json",
 					"JsonUtilities",
 					"PakFile",
diff --git a/Engine/Plugins/Runtime/GameFeatures/Source/GameFeatures/Private/IrisFilterGameFeatureAction.cpp b/Engine/Plugins/Runtime/GameFeatures/Source/GameFeatures/Private/IrisFilterGameFeatureAction.cpp
deleted file mode 100644
diff --git a/Engine/Plugins/Runtime/GameFeatures/Source/GameFeatures/Public/IrisFilterGameFeatureAction.h b/Engine/Plugins/Runtime/GameFeatures/Source/GameFeatures/Public/IrisFilterGameFeatureAction.h
deleted file mode 100644
