Clash milestone support
* Updating Rift Anomalies to be enabled by milestone
* Fixing some bugs related to initialization order of event router/gfas/etc.
* Fixing up wp GFA to also register edls and fixing up wind plugin to no longer be the one that loads things and use rivals instead.

#virtualized
#Synced-CL 51215150

diff --git a/Engine/Plugins/Runtime/GameFeatures/Source/GameFeatures/Private/GameFeatureAction_DataRegistrySource.cpp b/Engine/Plugins/Runtime/GameFeatures/Source/GameFeatures/Private/GameFeatureAction_DataRegistrySource.cpp
--- a/Engine/Plugins/Runtime/GameFeatures/Source/GameFeatures/Private/GameFeatureAction_DataRegistrySource.cpp
+++ b/Engine/Plugins/Runtime/GameFeatures/Source/GameFeatures/Private/GameFeatureAction_DataRegistrySource.cpp
@@ -53,10 +53,8 @@ void UGameFeatureAction_DataRegistrySource::OnGameFeatureRegistering()
 	}
 }
 
-void UGameFeatureAction_DataRegistrySource::OnGameFeatureActivating()
+void UGameFeatureAction_DataRegistrySource::RegisterRuntimeSources()
 {
-	Super::OnGameFeatureActivating();
-
 	if (ShouldPreloadAtRegistration())
 	{
 		// This already happened at registration
@@ -95,12 +93,12 @@ void UGameFeatureAction_DataRegistrySource::OnGameFeatureActivating()
 				{
 					if (!RegistrySource.DataTableToAdd.IsNull() && !RegistrySource.DataTableToAdd.IsValid())
 					{
-						UE_LOG(LogGameFeatures, Log, TEXT("OnGameFeatureActivating %s: DataRegistry source asset %s was not loaded before activation, this may cause a long hitch"), *GetPathName(), *RegistrySource.DataTableToAdd.ToString())
+						UE_LOG(LogGameFeatures, Log, TEXT("RegisterRuntimeSources %s: DataRegistry source asset %s was not loaded before activation, this may cause a long hitch"), *GetPathName(), *RegistrySource.DataTableToAdd.ToString())
 					}
 
 					if (!RegistrySource.CurveTableToAdd.IsNull() && !RegistrySource.CurveTableToAdd.IsValid())
 					{
-						UE_LOG(LogGameFeatures, Log, TEXT("OnGameFeatureActivating %s: DataRegistry source asset %s was not loaded before activation, this may cause a long hitch"), *GetPathName(), *RegistrySource.DataTableToAdd.ToString())
+						UE_LOG(LogGameFeatures, Log, TEXT("RegisterRuntimeSources %s: DataRegistry source asset %s was not loaded before activation, this may cause a long hitch"), *GetPathName(), *RegistrySource.DataTableToAdd.ToString())
 					}
 				}
 #endif
@@ -112,46 +110,58 @@ void UGameFeatureAction_DataRegistrySource::OnGameFeatureActivating()
 	}
 }
 
-void UGameFeatureAction_DataRegistrySource::OnGameFeatureUnregistering()
+void UGameFeatureAction_DataRegistrySource::OnGameFeatureActivating(FGameFeatureActivatingContext& Context)
 {
-	Super::OnGameFeatureUnregistering();
+	Super::OnGameFeatureActivating(Context);
 
-	if (ShouldPreloadAtRegistration())
+	RegisterRuntimeSources();
+}
+	
+void UGameFeatureAction_DataRegistrySource::UnregisterSources()
+{
+	UDataRegistrySubsystem* DataRegistrySubsystem = UDataRegistrySubsystem::Get();
+	if (ensure(DataRegistrySubsystem))
 	{
-		// This should only happen when the user is manually changing phase via the feature editor UI
-		UDataRegistrySubsystem* DataRegistrySubsystem = UDataRegistrySubsystem::Get();
-		if (ensure(DataRegistrySubsystem))
+		for (const FDataRegistrySourceToAdd& RegistrySource : SourcesToAdd)
 		{
-			for (const FDataRegistrySourceToAdd& RegistrySource : SourcesToAdd)
+			if (!RegistrySource.DataTableToAdd.IsNull())
 			{
-				if (!RegistrySource.DataTableToAdd.IsNull())
+				if (!DataRegistrySubsystem->UnregisterSpecificAsset(RegistrySource.RegistryToAddTo, RegistrySource.DataTableToAdd.ToSoftObjectPath()))
 				{
-					if (!DataRegistrySubsystem->UnregisterSpecificAsset(RegistrySource.RegistryToAddTo, RegistrySource.DataTableToAdd.ToSoftObjectPath()))
-					{
-						UE_LOG(LogGameFeatures, Log, TEXT("OnGameFeatureUnregistering %s: DataRegistry data table %s failed to unregister"), *GetPathName(), *RegistrySource.DataTableToAdd.ToString())
-					}
-					else
-					{
-						UE_LOG(LogGameFeatures, Log, TEXT("OnGameFeatureUnregistering %s: Temporarily disabling preloaded data table %s"), *GetPathName(), *RegistrySource.DataTableToAdd.ToString())
-					}
+					UE_LOG(LogGameFeatures, Log, TEXT("UnregisterSources %s: DataRegistry data table %s failed to unregister"), *GetPathName(), *RegistrySource.DataTableToAdd.ToString())
 				}
+				else
+				{
+					UE_LOG(LogGameFeatures, Log, TEXT("UnregisterSources %s: Temporarily disabling data table %s"), *GetPathName(), *RegistrySource.DataTableToAdd.ToString())
+				}
+			}
 
-				if (!RegistrySource.CurveTableToAdd.IsNull())
+			if (!RegistrySource.CurveTableToAdd.IsNull())
+			{
+				if (!DataRegistrySubsystem->UnregisterSpecificAsset(RegistrySource.RegistryToAddTo, RegistrySource.CurveTableToAdd.ToSoftObjectPath()))
 				{
-					if (!DataRegistrySubsystem->UnregisterSpecificAsset(RegistrySource.RegistryToAddTo, RegistrySource.CurveTableToAdd.ToSoftObjectPath()))
-					{
-						UE_LOG(LogGameFeatures, Log, TEXT("OnGameFeatureUnregistering %s: DataRegistry curve table %s failed to unregister"), *GetPathName(), *RegistrySource.CurveTableToAdd.ToString())
-					}
-					else
-					{
-						UE_LOG(LogGameFeatures, Log, TEXT("OnGameFeatureUnregistering %s: Temporarily disabling preloaded curve table %s"), *GetPathName(), *RegistrySource.CurveTableToAdd.ToString())
-					}
+					UE_LOG(LogGameFeatures, Log, TEXT("UnregisterSources %s: DataRegistry curve table %s failed to unregister"), *GetPathName(), *RegistrySource.CurveTableToAdd.ToString())
+				}
+				else
+				{
+					UE_LOG(LogGameFeatures, Log, TEXT("UnregisterSources %s: Temporarily disabling curve table %s"), *GetPathName(), *RegistrySource.CurveTableToAdd.ToString())
 				}
 			}
 		}
 	}
 }
 
+void UGameFeatureAction_DataRegistrySource::OnGameFeatureUnregistering()
+{
+	Super::OnGameFeatureUnregistering();
+
+	if (ShouldPreloadAtRegistration())
+	{
+		// This should only happen when the user is manually changing phase via the feature editor UI
+		UnregisterSources();
+	}
+}
+
 void UGameFeatureAction_DataRegistrySource::OnGameFeatureDeactivating(FGameFeatureDeactivatingContext& Context)
 {
 	Super::OnGameFeatureDeactivating(Context);
@@ -162,28 +172,7 @@ void UGameFeatureAction_DataRegistrySource::OnGameFeatureDeactivating(FGameFeatu
 		return;
 	}
 
-	UDataRegistrySubsystem* DataRegistrySubsystem = UDataRegistrySubsystem::Get();
-	if (ensure(DataRegistrySubsystem))
-	{
-		for (const FDataRegistrySourceToAdd& RegistrySource : SourcesToAdd)
-		{
-			if (!RegistrySource.DataTableToAdd.IsNull())
-			{
-				if (!DataRegistrySubsystem->UnregisterSpecificAsset(RegistrySource.RegistryToAddTo, RegistrySource.DataTableToAdd.ToSoftObjectPath()))
-				{
-					UE_LOG(LogGameFeatures, Log, TEXT("OnGameFeatureDeactivating %s: DataRegistry data table %s failed to unregister"), *GetPathName(), *RegistrySource.DataTableToAdd.ToString())
-				}
-			}
-
-			if (!RegistrySource.CurveTableToAdd.IsNull())
-			{
-				if (!DataRegistrySubsystem->UnregisterSpecificAsset(RegistrySource.RegistryToAddTo, RegistrySource.CurveTableToAdd.ToSoftObjectPath()))
-				{
-					UE_LOG(LogGameFeatures, Log, TEXT("OnGameFeatureDeactivating %s: DataRegistry curve table %s failed to unregister"), *GetPathName(), *RegistrySource.CurveTableToAdd.ToString())
-				}
-			}
-		}
-	}
+	UnregisterSources();
 }
 
 bool UGameFeatureAction_DataRegistrySource::ShouldPreloadAtRegistration()
diff --git a/Engine/Plugins/Runtime/GameFeatures/Source/GameFeatures/Public/GameFeatureAction_DataRegistrySource.h b/Engine/Plugins/Runtime/GameFeatures/Source/GameFeatures/Public/GameFeatureAction_DataRegistrySource.h
--- a/Engine/Plugins/Runtime/GameFeatures/Source/GameFeatures/Public/GameFeatureAction_DataRegistrySource.h
+++ b/Engine/Plugins/Runtime/GameFeatures/Source/GameFeatures/Public/GameFeatureAction_DataRegistrySource.h
@@ -52,21 +52,21 @@ class UGameFeatureAction_DataRegistrySource : public UGameFeatureAction
 	GENERATED_BODY()
 
 public:
-	virtual void OnGameFeatureRegistering() override;
-	virtual void OnGameFeatureUnregistering() override;
-	virtual void OnGameFeatureActivating() override;
-	virtual void OnGameFeatureDeactivating(FGameFeatureDeactivatingContext& Context) override;
+	GAMEFEATURES_API virtual void OnGameFeatureRegistering() override;
+	GAMEFEATURES_API virtual void OnGameFeatureUnregistering() override;
+	GAMEFEATURES_API virtual void OnGameFeatureActivating(FGameFeatureActivatingContext& Context) override;
+	GAMEFEATURES_API virtual void OnGameFeatureDeactivating(FGameFeatureDeactivatingContext& Context) override;
 
 	/** If true, we should load the sources at registration time instead of activation time */
-	virtual bool ShouldPreloadAtRegistration();
+	GAMEFEATURES_API virtual bool ShouldPreloadAtRegistration();
 
 #if WITH_EDITORONLY_DATA
-	virtual void AddAdditionalAssetBundleData(FAssetBundleData& AssetBundleData) override;
+	GAMEFEATURES_API virtual void AddAdditionalAssetBundleData(FAssetBundleData& AssetBundleData) override;
 #endif
 
 	//~UObject interface
 #if WITH_EDITOR
-	virtual EDataValidationResult IsDataValid(class FDataValidationContext& Context) const override;
+	GAMEFEATURES_API virtual EDataValidationResult IsDataValid(class FDataValidationContext& Context) const override;
 #endif
 	//~End of UObject interface
 
@@ -75,7 +75,13 @@ public:
 	GAMEFEATURES_API void AddSource(const FDataRegistrySourceToAdd& NewSource);
 #endif
 
+protected:
+	GAMEFEATURES_API void RegisterRuntimeSources();
+	
+	GAMEFEATURES_API void UnregisterSources();
+	
 private:
+	
 	/** List of sources to add when this feature is activated */
 	UPROPERTY(EditAnywhere, Category = "Registry Data", meta=(TitleProperty="RegistryToAddTo"))
 	TArray<FDataRegistrySourceToAdd> SourcesToAdd;
