Fixing world partition milestones

#changelist validated
#virtualized
#Synced-CL 51322170

diff --git a/Engine/Plugins/Runtime/GameFeatures/Source/GameFeatures/Private/GameFeatureAction_AddWorldPartitionContent.cpp b/Engine/Plugins/Runtime/GameFeatures/Source/GameFeatures/Private/GameFeatureAction_AddWorldPartitionContent.cpp
--- a/Engine/Plugins/Runtime/GameFeatures/Source/GameFeatures/Private/GameFeatureAction_AddWorldPartitionContent.cpp
+++ b/Engine/Plugins/Runtime/GameFeatures/Source/GameFeatures/Private/GameFeatureAction_AddWorldPartitionContent.cpp
@@ -35,7 +35,7 @@ void UGameFeatureAction_AddWorldPartitionContent::OnGameFeatureUnregistering()
 	Super::OnGameFeatureUnregistering();
 }
 
-void UGameFeatureAction_AddWorldPartitionContent::OnGameFeatureActivating()
+void UGameFeatureAction_AddWorldPartitionContent::OnGameFeatureActivating(FGameFeatureActivatingContext& Context)
 {
 	Super::OnGameFeatureActivating();
 	if (ExternalDataLayerAsset)
diff --git a/Engine/Plugins/Runtime/GameFeatures/Source/GameFeatures/Public/GameFeatureAction_AddWorldPartitionContent.h b/Engine/Plugins/Runtime/GameFeatures/Source/GameFeatures/Public/GameFeatureAction_AddWorldPartitionContent.h
--- a/Engine/Plugins/Runtime/GameFeatures/Source/GameFeatures/Public/GameFeatureAction_AddWorldPartitionContent.h
+++ b/Engine/Plugins/Runtime/GameFeatures/Source/GameFeatures/Public/GameFeatureAction_AddWorldPartitionContent.h
@@ -21,7 +21,7 @@ public:
 	//~ Begin UGameFeatureAction interface
 	UE_API virtual void OnGameFeatureRegistering() override;
 	UE_API virtual void OnGameFeatureUnregistering() override;
-	UE_API virtual void OnGameFeatureActivating() override;
+	UE_API virtual void OnGameFeatureActivating(FGameFeatureActivatingContext& Context) override;
 	UE_API virtual void OnGameFeatureDeactivating(FGameFeatureDeactivatingContext& Context) override;
 	//~ End UGameFeatureAction interface
 
@@ -36,6 +36,11 @@ public:
 	const TObjectPtr<const UExternalDataLayerAsset>& GetExternalDataLayerAsset() const { return ExternalDataLayerAsset; }
 #endif
 
+protected:
+	/** External Data Layer used by this action. */
+	UPROPERTY(EditAnywhere, Category = DataLayer)
+	TObjectPtr<const UExternalDataLayerAsset> ExternalDataLayerAsset;
+	
 private:
 #if WITH_EDITOR
 	UE_API void OnExternalDataLayerAssetChanged(const UExternalDataLayerAsset* OldAsset, const UExternalDataLayerAsset* NewAsset);
@@ -45,10 +50,6 @@ private:
 	TWeakObjectPtr<const UExternalDataLayerAsset> PreEditChangeExternalDataLayerAsset;
 	TWeakObjectPtr<const UExternalDataLayerAsset> PreEditUndoExternalDataLayerAsset;
 
-	/** External Data Layer used by this action. */
-	UPROPERTY(EditAnywhere, Category = DataLayer)
-	TObjectPtr<const UExternalDataLayerAsset> ExternalDataLayerAsset;
-
 #if WITH_EDITORONLY_DATA
 	/** Only used when converting from UGameFeatureAction_AddWPContent */
 	UPROPERTY()
