1) GFP_B was SetDest -> Active
2) GFP_B gets canceled back to -> Registered, we set bTryCancel=true and we are letting it trying to finish that transition
3) GFP_B has a parent GFP_A trying to Activte which requires its depends to now be Active, so it iterates through an tries to TransitionDepends on all its children

GFP_B has bTryCancel=True and it may be at Dest state which will fire a callback that it Completed a transition with an Error = Canceled

In these cases, we have canceled, the cancel in a sense, and should just try to move that GFP to the state the Parent requires by re-trying a TransitionDepends only once

#Synced-CL 51322414

#Synced-CL 51322414

diff --git a/Engine/Plugins/Runtime/GameFeatures/Source/GameFeatures/Private/GameFeaturePluginStateMachine.cpp b/Engine/Plugins/Runtime/GameFeatures/Source/GameFeatures/Private/GameFeaturePluginStateMachine.cpp
--- a/Engine/Plugins/Runtime/GameFeatures/Source/GameFeatures/Private/GameFeaturePluginStateMachine.cpp
+++ b/Engine/Plugins/Runtime/GameFeatures/Source/GameFeatures/Private/GameFeaturePluginStateMachine.cpp
@@ -1090,9 +1090,9 @@ struct FTransitionDependenciesGameFeaturePluginState : public FGameFeaturePlugin
 			}
 		}
 
-		for (FDepResultPair& Pair : RemainingDependencies)
+		for (FDepResult& Dep : RemainingDependencies)
 		{
-			UGameFeaturePluginStateMachine* RemainingDependency = Pair.Key.Get();
+			UGameFeaturePluginStateMachine* RemainingDependency = Dep.Machine.Get();
 			if (!RemainingDependency)
 			{
 				// One of the dependency state machines was destroyed before finishing
@@ -1100,14 +1100,14 @@ struct FTransitionDependenciesGameFeaturePluginState : public FGameFeaturePlugin
 				return;
 			}
 
-			if (Pair.Value.HasError())
+			if (Dep.Result.HasError())
 			{
 				const FStringView ShortUrl = StateProperties.PluginIdentifier.GetIdentifyingString();
 				const FStringView DepShortUrl = GameFeaturesSubsystem.GetPluginIdentifier(RemainingDependency).GetIdentifyingString();
 				UE_LOG(LogGameFeatures, Error, TEXT("Plugin (%.*s) dependency (%.*s) failed to transition with error %s"),
 					ShortUrl.Len(), ShortUrl.GetData(),
 					DepShortUrl.Len(), DepShortUrl.GetData(),
-					*Pair.Value.GetError());
+					*Dep.Result.GetError());
 				StateStatus.SetTransitionError(TransitionPolicy::GetErrorState(), GetErrorResult(TEXT("Failed_Dependency_Transition")));
 				if (GameFeaturesSubsystem.IsErrorStateUnrecoverable(RemainingDependency))
 				{
@@ -1182,14 +1182,14 @@ struct FTransitionDependenciesGameFeaturePluginState : public FGameFeaturePlugin
 			UGameFeaturePluginStateMachine* NewMachine = UGameFeaturesSubsystem::Get().FindOrCreateGameFeaturePluginStateMachine(GameFeaturesSubsystem.GetPluginURL(Dependency), DepProtocolOptions);
 			checkf(NewMachine != Dependency, TEXT("Game Feature Plugin %s should have already been removed from subsystem!"), *GameFeaturesSubsystem.GetPluginURL(Dependency));
 
-			const int32 Index = RemainingDependencies.IndexOfByPredicate([Dependency](const FDepResultPair& Pair)
+			const int32 Index = RemainingDependencies.IndexOfByPredicate([Dependency](const FDepResult& Dep)
 				{
-					return Pair.Key == Dependency;
+					return Dep.Machine == Dependency;
 				});
 
 			check(Index != INDEX_NONE);
-			FDepResultPair& FoundDep = RemainingDependencies[Index];
-			FoundDep.Key = NewMachine;
+			FDepResult& FoundDep = RemainingDependencies[Index];
+			FoundDep.Machine = NewMachine;
 
 			GameFeaturesSubsystem.RemovePendingTransitionCallback(Dependency, this);
 			GameFeaturesSubsystem.RemovePendingCancelCallback(Dependency, this);
@@ -1222,18 +1222,29 @@ struct FTransitionDependenciesGameFeaturePluginState : public FGameFeaturePlugin
 
 	void OnDependencyTransitionComplete(UGameFeaturePluginStateMachine* Dependency, const UE::GameFeatures::FResult& Result)
 	{
-		const int32 Index = RemainingDependencies.IndexOfByPredicate([Dependency](const FDepResultPair& Pair)
+		const int32 Index = RemainingDependencies.IndexOfByPredicate([Dependency](const FDepResult& Dep)
 			{
-				return Pair.Key == Dependency;
+				return Dep.Machine == Dependency;
 			});
 
 		if (ensure(Index != INDEX_NONE))
 		{
-			FDepResultPair& FoundDep = RemainingDependencies[Index];
+			FDepResult& FoundDep = RemainingDependencies[Index];
 
 			if (Result.HasError())
 			{
-				FoundDep.Value = Result;
+				// If our error is a cancel, retry once — the dependency may have been canceled by a competing
+				// request and is now at a stable state we can transition from (e.g. Registered after a TryCancel
+				// settled there). Without this retry, we'd end up stuck in ErrorActivatingDependencies even
+				// though the dependency is perfectly healthy.
+				if (Result.GetError() == UE::GameFeatures::CanceledResult.GetError() && !FoundDep.bRetriedAfterCancel)
+				{
+					FoundDep.bRetriedAfterCancel = true;
+					TransitionDependency(Dependency);
+					return;
+				}
+
+				FoundDep.Result = Result;
 			}
 			else
 			{
@@ -1270,9 +1281,9 @@ struct FTransitionDependenciesGameFeaturePluginState : public FGameFeaturePlugin
 		if (RemainingDependencies.Num() > 0)
 		{
 			UGameFeaturesSubsystem& GameFeaturesSubsystem = UGameFeaturesSubsystem::Get();
-			for (FDepResultPair& Pair : RemainingDependencies)
+			for (FDepResult& Dep : RemainingDependencies)
 			{
-				UGameFeaturePluginStateMachine* RemainingDependency = Pair.Key.Get();
+				UGameFeaturePluginStateMachine* RemainingDependency = Dep.Machine.Get();
 				if (RemainingDependency)
 				{
 					GameFeaturesSubsystem.RemovePendingTransitionCallback(RemainingDependency, this);
@@ -1290,8 +1301,16 @@ struct FTransitionDependenciesGameFeaturePluginState : public FGameFeaturePlugin
 		bRequestedDependencies = false;
 	}
 
-	using FDepResultPair = TPair<TWeakObjectPtr<UGameFeaturePluginStateMachine>, UE::GameFeatures::FResult>;
-	TArray<FDepResultPair> RemainingDependencies;
+	struct FDepResult
+	{
+		FDepResult(UGameFeaturePluginStateMachine* InMachine, UE::GameFeatures::FResult InResult)
+			: Machine(InMachine), Result(MoveTemp(InResult)) {}
+
+		TWeakObjectPtr<UGameFeaturePluginStateMachine> Machine;
+		UE::GameFeatures::FResult Result;
+		bool bRetriedAfterCancel = false;
+	};
+	TArray<FDepResult> RemainingDependencies;
 	bool bRequestedDependencies = false;
 	bool bCheckedRealtimeMode = false;
 };
