Add method UGameFeaturesProjectPolicies::GetThreadSafetyVersion() which returns an enum telling the GameFeatures system which methods can be called on any thread.
Initial version requires ShouldReadPluginDetails, GetGameFeaturePluginURL and IsPluginAllowed to be callable on any thread.
Base implementation returns Initial for UGameFeaturesProjectPolicies and UDefaultGameFeaturesProjectPolicies but None if the policies have been subclassed, even if they didn't override GetThreadSafetyVersion.

Document thread safety of some methods in UGameFeaturesSubsystem.

#Synced-CL 51619035

diff --git a/Engine/Plugins/Runtime/GameFeatures/Source/GameFeatures/Private/GameFeaturesProjectPolicies.cpp b/Engine/Plugins/Runtime/GameFeatures/Source/GameFeatures/Private/GameFeaturesProjectPolicies.cpp
--- a/Engine/Plugins/Runtime/GameFeatures/Source/GameFeatures/Private/GameFeaturesProjectPolicies.cpp
+++ b/Engine/Plugins/Runtime/GameFeatures/Source/GameFeatures/Private/GameFeaturesProjectPolicies.cpp
@@ -21,6 +21,17 @@ namespace UE::GameFeatures
 	}
 }
 
+
+EGameFeaturesProjectPoliciesThreadSafetyVersion UGameFeaturesProjectPolicies::GetThreadSafetyVersion() const
+{
+	if (   GetClass() == UGameFeaturesProjectPolicies::StaticClass()
+		|| GetClass() == UDefaultGameFeaturesProjectPolicies::StaticClass())
+	{
+		return EGameFeaturesProjectPoliciesThreadSafetyVersion::Initial;
+	}
+	return EGameFeaturesProjectPoliciesThreadSafetyVersion::None;
+}
+
 void UDefaultGameFeaturesProjectPolicies::InitGameFeatureManager()
 {
 	UE_LOG(LogGameFeatures, Log, TEXT("Scanning for built-in game feature plugins"));
@@ -88,7 +99,7 @@ bool UGameFeaturesProjectPolicies::GetGameFeaturePluginURL(const TSharedRef<IPlu
 	const FString& PluginDescriptorFilename = Plugin->GetDescriptorFileName();
 	if (!PluginDescriptorFilename.IsEmpty())
 	{
-		if (GetDefault<UGameFeaturesSubsystemSettings>()->IsValidGameFeaturePlugin(FPaths::ConvertRelativePathToFull(PluginDescriptorFilename)))
+		if (GetDefault<UGameFeaturesSubsystemSettings>()->IsValidGameFeaturePlugin(PluginDescriptorFilename))
 		{
 			OutPluginURL = UGameFeaturesSubsystem::GetPluginURL_FileProtocol(PluginDescriptorFilename);
 		}
diff --git a/Engine/Plugins/Runtime/GameFeatures/Source/GameFeatures/Public/GameFeaturesProjectPolicies.h b/Engine/Plugins/Runtime/GameFeatures/Source/GameFeatures/Public/GameFeaturesProjectPolicies.h
--- a/Engine/Plugins/Runtime/GameFeatures/Source/GameFeatures/Public/GameFeaturesProjectPolicies.h
+++ b/Engine/Plugins/Runtime/GameFeatures/Source/GameFeatures/Public/GameFeaturesProjectPolicies.h
@@ -2,6 +2,7 @@
 
 #pragma once
 
+#include "Async/TransactionallySafeRecursiveMutex.h"
 #include "UObject/PrimaryAssetId.h"
 #include "GameFeaturesSubsystem.h"
 
@@ -27,6 +28,19 @@ enum class EStreamingAssetInstallMode : uint8
 	Full, // stream in all data
 };
 
+// Specification of thread safety contract for a class derived from UGameFeaturesProjectPolicies.
+// @see UGameFeaturesProjectPolicies::GetThreadSafetyVersion
+enum class EGameFeaturesProjectPoliciesThreadSafetyVersion
+{
+	// No methods are threadsafe, default for derived classes
+	None = 0, 
+	// Initial version when this enum was added. The follow methods must be callable from any thread:
+	//	ShouldReadPluginDetails
+	// 	GetGameFeaturePluginURL
+	//  IsPluginAllowed
+	Initial,
+};
+
 // This class allows project-specific rules to be implemented for game feature plugins.
 // Create a subclass and choose it in Project Settings .. Game Features
 UCLASS(MinimalAPI)
@@ -41,6 +55,23 @@ public:
 	// Called when the game feature manager is shut down
 	virtual void ShutdownGameFeatureManager() { }
 
+	/**
+	 * Returns which methods of this class may be called from any thread.
+	 * The methods which must be callable from any thread for each version are marked as such in this file.
+	 * Derived classes should return the maximum value of the enum when they were written.
+	 * The base implementation returns the maximum current version if this is not a derived class.
+	 *
+	 * Specifically, implementations should expect that these methods may be called concurrently with calls 
+	 * the game makes into UGameFeaturesSubsystem, and during the subsystem's own ticks.
+	 * The subsystem makes no particular guarantees about avoid reentrancy if these methods call into UGameFeaturesSubsystem.
+	 * e.g. if the implementation of IsPluginAllowed calls into the subsystem, the subsystem may call back into another method such as
+	 * GetGameFeaturePluginURL. The possibility of priority inversion / deadlocks should be considered.
+	 *
+	 * Note that no implementation needs to support methods being called concurrently with InitGameFeatureManager
+	 * or ShutdownGameFeatureManager unless otherwise specified.
+	 */
+	UE_API virtual EGameFeaturesProjectPoliciesThreadSafetyVersion GetThreadSafetyVersion() const;
+
 	// Called to determined the expected state of a plugin under the WhenLoading conditions.
 	UE_API virtual bool WillPluginBeCooked(const FString& PluginFilename, const FGameFeaturePluginDetails& PluginDetails) const;
 
@@ -61,25 +92,53 @@ public:
 	// This defaults to true for the first few frames of a normal game or editor, but can be overridden.
 	UE_API virtual bool IsLoadingStartupPlugins() const;
 
-	// Called to determine the plugin URL for a given known Plugin. Can be used if the policy wants to deliver non file based URLs.
+	/**
+	  * Called to determine the plugin URL for a given known Plugin. Can be used if the policy wants to deliver non file based URLs.
+	  *
+	  * THREAD SAFETY: Should be threadsafe for EGameFeaturesProjectPoliciesThreadSafetyVersion::Initial
+	  */
 	UE_API virtual bool GetGameFeaturePluginURL(const TSharedRef<IPlugin>& Plugin, FString& OutPluginURL) const;
 
 	UE_DEPRECATED(5.6, "Replaced with IsPluginAllowed(PluginURL, OutReason).")
 	virtual bool IsPluginAllowed(const FString& PluginURL) const { return true; }
-	// Called to determine if a plugin is allowed to be loaded or not
-	// (e.g., when doing a fast cook a game might want to disable some or all game feature plugins)
+
+	/**
+	  * Called to determine if a plugin is allowed to be loaded or not
+	  * (e.g., when doing a fast cook a game might want to disable some or all game feature plugins)
+	  *
+	  * THREAD SAFETY: Should be threadsafe for EGameFeaturesProjectPoliciesThreadSafetyVersion::Initial
+	  */
 	virtual bool IsPluginAllowed(const FString& PluginURL, FString* OutReason) const 
 	{ 
 		if (OutReason)
 		{
-			OutReason = {};
+			*OutReason = {};
 		}
 		PRAGMA_DISABLE_DEPRECATION_WARNINGS
 		return IsPluginAllowed(PluginURL); 
 		PRAGMA_ENABLE_DEPRECATION_WARNINGS
 	}
 
-	// Return true if a uplugin's details should be read and false if it should be skipped. Skipped plugins will not be processed as GFPs and skipped as though they didn't exist. Useful to limit the number of uplugin files opened for perf reasons
+	/**
+	  * Called to determine if a plugin is allowed to be loaded or not.
+	  * (e.g., when doing a fast cook a game might want to disable some or all game feature plugins)
+	  * New version that provides plugin details when they are already available to avoid re-entrant calls.
+	  * Prior version is not yet deprecated.
+	  *
+	  * THREAD SAFETY: Should be threadsafe for EGameFeaturesProjectPoliciesThreadSafetyVersion::Initial
+	  */
+	virtual bool IsPluginAllowed(const FString& PluginURL, const FGameFeaturePluginDetails* InDetails, FString* OutReason) const 
+	{ 
+		return IsPluginAllowed(PluginURL, OutReason); 
+	}
+
+	/**
+	  * Return true if a uplugin's details should be read and false if it should be skipped.
+	  * Skipped plugins will not be processed as GFPs and skipped as though they didn't exist.
+	  * Useful to limit the number of uplugin files opened for perf reasons.
+	  *
+	  * THREAD SAFETY: Should be threadsafe for EGameFeaturesProjectPoliciesThreadSafetyVersion::Initial
+	  */
 	virtual bool ShouldReadPluginDetails(const FString& PluginDescriptorFilename) const { return true; }
 
 	// Return true if a plugin should be allowed to register if it doesn't have a GameFeatureData. Useful in cases where you want to allow the content of certain plugins to not exist.
@@ -119,6 +178,10 @@ public:
 	 * @param - bEvenIfDoesntExist - when true will return the name of bundle we are looking for without checking if it exists or not.
 	 */
 	UE_API virtual FString GetOptionalInstallBundleName(FStringView PluginName, bool bEvenIfDoesntExist = false);
+
+protected:
+	/** This mutex can be used by derived classes for broad mutual exclusion */
+	UE::FTransactionallySafeRecursiveMutex Mutex;
 };
 
 // This is a default implementation that immediately processes all game feature plugins the based on
diff --git a/Engine/Plugins/Runtime/GameFeatures/Source/GameFeatures/Public/GameFeaturesSubsystem.h b/Engine/Plugins/Runtime/GameFeatures/Source/GameFeatures/Public/GameFeaturesSubsystem.h
--- a/Engine/Plugins/Runtime/GameFeatures/Source/GameFeatures/Public/GameFeaturesSubsystem.h
+++ b/Engine/Plugins/Runtime/GameFeatures/Source/GameFeatures/Public/GameFeaturesSubsystem.h
@@ -549,7 +549,10 @@ public:
 	}
 
 public:
-	/** Construct a 'file:' Plugin URL using from the PluginDescriptorPath */
+	/**
+	  * Construct a 'file:' Plugin URL using from the PluginDescriptorPath 
+	  * THREAD SAFETY: Any thread
+	  */
 	static UE_API FString GetPluginURL_FileProtocol(const FString& PluginDescriptorPath);
 	static UE_API FString GetPluginURL_FileProtocol(const FString& PluginDescriptorPath, TArrayView<const TPair<FString, FString>> AdditionalOptions);
 	
@@ -569,7 +572,10 @@ public:
 	/** Tests whether the plugin URL is the specified protocol */
 	static UE_API bool IsPluginURLProtocol(FStringView PluginURL, EGameFeaturePluginProtocol PluginProtocol);
 
-	/** Parse the plugin URL into subparts */
+	/** 
+	  * Parse the plugin URL into subparts 
+	  * THREAD SAFETY: Any thread
+	  */
 	static UE_API bool ParsePluginURL(FStringView PluginURL, EGameFeaturePluginProtocol* OutProtocol = nullptr, FStringView* OutPath = nullptr, FStringView* OutOptions = nullptr);
 
 	/** Parse options from a plugin URL or the options subpart of the plugin URL */
@@ -682,6 +688,7 @@ public:
 	/**
 	 * If the specified plugin is known by the game feature system, returns the URL used to identify it
 	 * @return true if the plugin exists, false if it was not found
+	 * THREAD SAFETY: Any thread as long as GameSpecificPolicies meets version EGameFeaturesProjectPoliciesThreadSafetyVersion::Initial
 	 */
 	UE_API bool GetPluginURLByName(FStringView PluginName, FString& OutPluginURL) const;
 	UE_API bool GetPluginURLByName(FUtf8StringView PluginName, FString& OutPluginURL) const;
@@ -762,6 +769,10 @@ public:
 	UE_API void LoadAllBuiltInGameFeaturePluginDetails(bool bForce = false);
 
 	/** Gets the URL for the given plugin, applying game-specific policies where appropriate. Should only be used for built-in GFPs */
+	/** 
+	  * Gets the URL for the given plugin, applying game-specific policies where appropriate. Should only be used for built-in GFPs 
+	  * THREAD SAFETY: Any thread as long as GameSpecificPolicies meets version EGameFeaturesProjectPoliciesThreadSafetyVersion::Initial
+	  */
 	UE_API bool GetBuiltInGameFeaturePluginURL(const TSharedRef<IPlugin>& Plugin, FString& OutPluginURL) const;
 
 	/** Gets relevant properties out of a uplugin file if it's installed */
@@ -774,8 +785,16 @@ public:
 	 */
 	UE_API bool GetGameFeatureControlsUPlugin(const FString& PluginURL, bool& OutGameFeatureControlsUPlugin) const;
 
+	/**
+	 * THREAD SAFETY: Game thread only
+	 */
 	UE_API bool IsPluginAllowed(const FString& PluginURL, FString* OutReason = nullptr) const;
 
+	/**
+	 * THREAD SAFETY: Any thread
+	 */
+	UE_API bool IsPluginAllowed(const FString& PluginURL, const FGameFeaturePluginDetails* InDetails, FString* OutReason) const;
+
 	/** 
 	 * Pre-install any required game feature data, which can be useful for larger payloads. 
 	 * This does not instantiate any GFP although it is safe to do so before this finishes. 
@@ -790,6 +809,9 @@ public:
 	static UE_API EBuiltInAutoState DetermineBuiltInInitialFeatureState(TSharedPtr<FJsonObject> Descriptor, const FString& ErrorContext);
 	static UE_API EBuiltInAutoState DetermineBuiltInInitialFeatureState(UE::Json::FConstObject Descriptor, const FString& ErrorContext);
 
+	/**
+	 * THREAD SAFETY: Any thread.
+	 */
 	static UE_API EGameFeaturePluginState ConvertInitialFeatureStateToTargetState(EBuiltInAutoState InitialState);
 
 	/** Used during a DLC cook to determine which plugins should be cooked */
@@ -882,7 +904,14 @@ private:
 
 private:
 
+	/** 
+	  * THREA SAFETY: Any thread
+	  */
 	UE_API bool ShouldUpdatePluginProtocolOptions(const UGameFeaturePluginStateMachine* StateMachine, const FGameFeatureProtocolOptions& NewOptions);
+	/** 
+	  * THREAD SAFETY: Any thread as long as no other thread is operating on that state machine.
+	  * e.g. from a parallel for on the game thread, as state machine updates only happen on the game thread.
+	  */
 	UE_API UE::GameFeatures::FResult UpdateGameFeatureProtocolOptions(UGameFeaturePluginStateMachine* StateMachine, const FGameFeatureProtocolOptions& NewOptions, bool* bOutDidUpdate = nullptr);
 
 	UE_API const UGameFeatureData* GetDataForStateMachine(UGameFeaturePluginStateMachine* GFSM) const;
@@ -906,13 +935,22 @@ private:
 	friend struct FGameFeaturePluginState_Unmounting;
 	friend struct FBaseDataReleaseGameFeaturePluginState;
 
-	/** Gets the state machine associated with the specified URL */
+	/**
+	  * Gets the state machine associated with the specified URL 
+	  * THREAD SAFETY: Any thread.
+	  */
 	UE_API UGameFeaturePluginStateMachine* FindGameFeaturePluginStateMachine(const FString& PluginURL) const;
 
-	/** Gets the state machine associated with the specified PluginIdentifier */
+	/**
+	  * Gets the state machine associated with the specified PluginIdentifier 
+	  * THREAD SAFETY: Any thread.
+	  */
 	UE_API UGameFeaturePluginStateMachine* FindGameFeaturePluginStateMachine(const FGameFeaturePluginIdentifier& PluginIdentifier) const;
 
-	/** Gets the state machine associated with the specified URL, creates it if it doesnt exist */
+	/**
+	  * Gets the state machine associated with the specified URL, creates it if it doesnt exist 
+	  * THREAD SAFETY: Any thread.
+	  */
 	UE_API UGameFeaturePluginStateMachine* FindOrCreateGameFeaturePluginStateMachine(const FString& PluginURL, const FGameFeatureProtocolOptions& ProtocolOptions, bool* bOutFoundExisting = nullptr);
 
 	/** Notification that a game feature has finished loading, and whether it was successful */
@@ -921,6 +959,7 @@ private:
 	/** 
 	 * Sets a new destination state. Will attempt to cancel the current transition if the new destination is incompatible with the current destination 
 	 * Note: In the case that the existing machine is terminal, a new one will need to be created. In that case ProtocolOptions will be used for the new machine.
+	 * THREAD SAFETY: Game thread.
 	 */
 	UE_API void ChangeGameFeatureDestination(UGameFeaturePluginStateMachine* Machine, const FGameFeaturePluginStateRange& StateRange, FGameFeaturePluginChangeStateComplete CompleteDelegate);
 	UE_API void ChangeGameFeatureDestination(UGameFeaturePluginStateMachine* Machine, const FGameFeatureProtocolOptions& ProtocolOptions, const FGameFeaturePluginStateRange& StateRange, FGameFeaturePluginChangeStateComplete CompleteDelegate);
