[GFP] Clear all primary assets loaded by a game from game feature plugin when the plugin unregisters.
Query in-memory assets matching the plugin's scan paths and unload them before removing the associated primary asset type scan paths

#tests local cooked build manual tests
#Synced-CL 51619035

diff --git a/Engine/Plugins/Runtime/GameFeatures/Source/GameFeatures/Private/GameFeaturesSubsystem.cpp b/Engine/Plugins/Runtime/GameFeatures/Source/GameFeatures/Private/GameFeaturesSubsystem.cpp
--- a/Engine/Plugins/Runtime/GameFeatures/Source/GameFeatures/Private/GameFeaturesSubsystem.cpp
+++ b/Engine/Plugins/Runtime/GameFeatures/Source/GameFeatures/Private/GameFeaturesSubsystem.cpp
@@ -148,6 +148,14 @@ namespace UE::GameFeatures
 		false,
 		TEXT("Suspend streaming manager in transition context to prevent potential GC races."));
 
+	static FAutoConsoleVariable CVarUnloadPrimaryAssetsOnUnregister(
+		TEXT("GameFeaturePlugin.UnloadPrimaryAssetsOnUnregister"),
+		false,
+		TEXT("Unload primary assets from AssetManager when the game feature plugin unregisters. ")
+		TEXT("Use this option only when you are certain that unmounting the plugin will not leak primary assets ")
+		TEXT("or their referenced secondary assets. ")
+		TEXT("This can provide a small performance improvement during plugin unmounting."));
+
 	TOptional<FString> GetPluginUrlForConsoleCommand(const TArray<FString>& Args, FOutputDevice& Ar)
 	{
 		TOptional<FString> PluginURL;
@@ -917,6 +925,56 @@ void UGameFeaturesSubsystem::RemoveGameFeatureFromAssetManager(const UGameFeatur
 	FString PluginRootPath = TEXT("/") + PluginName + TEXT("/");
 	UAssetManager& LocalAssetManager = UAssetManager::Get();
 
+	// Cleanup primary assets from asset manager if they were loaded by users directly.
+	if (UE::GameFeatures::CVarUnloadPrimaryAssetsOnUnregister->GetBool())
+	{
+		TArray<FString> ScanPaths;
+		if (GameFeatureToRemove->GetPrimaryAssetTypesToScan().Num() > 1)
+		{
+			int32 TotalSizeRequired = 0;
+			for (const FPrimaryAssetTypeInfo& Type : GameFeatureToRemove->GetPrimaryAssetTypesToScan())
+			{
+				TotalSizeRequired += Type.AssetScanPaths.Num();
+			}
+			ScanPaths.Reserve(TotalSizeRequired);
+		}
+		for (const FPrimaryAssetTypeInfo& Type : GameFeatureToRemove->GetPrimaryAssetTypesToScan())
+		{
+			ScanPaths.Append(Type.AssetScanPaths);
+		}
+
+		if (!ScanPaths.IsEmpty())
+		{
+			if (IAssetRegistry* AR = IAssetRegistry::Get())
+			{
+				FARFilter AssetsFilter;
+				for (const FString& Path : ScanPaths)
+				{
+					AssetsFilter.PackagePaths.Add(FName(Path));
+				}
+				AssetsFilter.bRecursivePaths = true;
+				AssetsFilter.TagsAndValues.Add(FPrimaryAssetId::PrimaryAssetTypeTag);
+				AssetsFilter.RecursiveClassPathsExclusionSet.Add(UGameFeatureData::StaticClass()->GetClassPathName());
+				AssetsFilter.bRecursiveClasses = true;
+
+				TArray<FAssetData> AssetDatas;
+				if (AR->GetInMemoryAssets(AssetsFilter, AssetDatas))
+				{
+					TArray<FPrimaryAssetId> PotentialAssets;
+					PotentialAssets.Reserve(AssetDatas.Num());
+
+					for (const FAssetData& Asset : AssetDatas)
+					{
+						PotentialAssets.Add(Asset.GetPrimaryAssetId());
+					}
+
+					UAssetManager::Get().UnloadPrimaryAssets(PotentialAssets);
+				}
+			}
+		}
+	}
+
+	// Remove scan paths from asset manager
 	for (FPrimaryAssetTypeInfo TypeInfo : GameFeatureToRemove->GetPrimaryAssetTypesToScan())
 	{
 		if (AddedPrimaryAssetTypes.Contains(TypeInfo.PrimaryAssetType))
