IAD shared now does not include m.ubulk files (excluded for now), added support for generating pakfilerules config and bundle config.

#Synced-CL 51662096

#Synced-CL 51662096

diff --git a/Engine/Plugins/Runtime/GameFeatures/Source/GameFeatures/Private/GameFeaturePluginStateMachine.cpp b/Engine/Plugins/Runtime/GameFeatures/Source/GameFeatures/Private/GameFeaturePluginStateMachine.cpp
--- a/Engine/Plugins/Runtime/GameFeatures/Source/GameFeatures/Private/GameFeaturePluginStateMachine.cpp
+++ b/Engine/Plugins/Runtime/GameFeatures/Source/GameFeatures/Private/GameFeaturePluginStateMachine.cpp
@@ -165,6 +165,48 @@ namespace UE::GameFeatures
 		false,
 		TEXT("Don't fail on missing on demand (IAD) asset dependencies"));
 
+
+	static TAutoConsoleVariable<bool> CVarForceGFPIADDownload(TEXT("GameFeaturePlugin.ForceGFPIADDownload"),
+		false,
+		TEXT("Enable batch processing when processing multiple plugins and specified in protocol options."));
+
+#if !UE_BUILD_SHIPPING
+	FString DebugGFPNameMatch = "";
+	static FAutoConsoleVariableRef CVarDebugGFPNameMatch(TEXT("GameFeaturePlugin.DebugGFPName"), DebugGFPNameMatch, TEXT("Set cvar and add breakpoint to \"ShouldDebugGFPNameMatch()\" To debug"));
+
+	bool ShouldDebugGFP(const FString& GFPName)
+	{
+		static FString CachedDebugName;
+		static TUniquePtr<FRegexPattern> Pattern;
+		if (CachedDebugName != DebugGFPNameMatch)
+		{
+			CachedDebugName = DebugGFPNameMatch;
+			if (DebugGFPNameMatch.IsEmpty())
+			{
+				Pattern = nullptr;
+			}
+			else
+			{
+				Pattern.Reset(new FRegexPattern(DebugGFPNameMatch));
+			}
+		}
+
+		if (Pattern.IsValid())
+		{
+			FRegexMatcher Matcher(*Pattern, GFPName);
+			if (Matcher.FindNext())
+			{
+				return true;
+			}
+		}
+		return false;
+	}
+#else
+	bool ShouldDebugGFP(const FString& GFPName)
+	{
+		return false;
+	}
+#endif
 	bool ShouldDeferLocalizationDataLoad()
 	{
 		// Note: We don't defer localization data loading in the editor, as the editor only needs to mount plugins to use them
@@ -1964,6 +2006,14 @@ struct FGameFeaturePluginState_Downloading : public FGameFeaturePluginState
 		IoStoreOnDemandContext = MakeShared<FIoStoreOnDemandContext>();
 		IoStoreOnDemandContext->PendingInstalls = MetaData.InstallBundlesWithAssetDependencies.Num();
 
+#if !UE_BUILD_SHIPPING
+		static bool bBreakOnInstallBundlesWithAssetDependencies = true;
+		if (bBreakOnInstallBundlesWithAssetDependencies && MetaData.InstallBundlesWithAssetDependencies.IsEmpty())
+		{
+			static int i = 0;
+			++i;
+		}
+#endif
 		for (const FName InstallBundle : MetaData.InstallBundlesWithAssetDependencies)
 		{
 			UE::IoStore::FOnDemandInstallArgs InstallArgs;
@@ -1971,13 +2021,14 @@ struct FGameFeaturePluginState_Downloading : public FGameFeaturePluginState
 			InstallArgs.TagSets.Emplace(TEXTVIEW("mount")); // May not be a real tagset, but this will install all the mandatory untagged chunks
 			InstallArgs.Options |= UE::IoStore::EOnDemandInstallOptions::InstallSoftReferences;
 			InstallArgs.Options |= UE::IoStore::EOnDemandInstallOptions::CallbackOnGameThread;
-			if (Options.bDoNotDownload)
+			if (Options.bDoNotDownload && !UE::GameFeatures::CVarForceGFPIADDownload.GetValueOnGameThread())
 			{
 				InstallArgs.Options |= UE::IoStore::EOnDemandInstallOptions::DoNotDownload;
 			}
 			
 			InstallArgs.ContentHandle = UE::GameFeatures::GFPSharedInstallTracker.AddOnDemandContentHandle(
 				InstallBundle, UE::GameFeatures::EGFPInstallLevel::Download);
+			InstallArgs.DebugName = StateProperties.PluginName;
 			check(InstallArgs.ContentHandle.IsValid());
 
 			// This should be pretty small, so not going to worry about progress here.
@@ -2095,7 +2146,7 @@ struct FGameFeaturePluginState_Downloading : public FGameFeaturePluginState
 
 		// Handle bDoNotDownload flag before doing any async ops
 		// if not up to date, check to see if we allow downloading
-		if (Options.bDoNotDownload && !bAllUpToDate)
+		if ((Options.bDoNotDownload && !UE::GameFeatures::CVarForceGFPIADDownload.GetValueOnGameThread()) && !bAllUpToDate)
 		{
 			Result = GetErrorResult(TEXT("GFPStateMachine.DownloadNotAllowed"));
 			bSuppressResultErrorLog = true; // Don't log an error if the user disallowed the download
@@ -3105,7 +3156,13 @@ struct FGameFeaturePluginState_AssetDependencyStreaming : public FGameFeaturePlu
 	virtual void BeginState() override
 	{
 		Cleanup();
-
+#if !UE_BUILD_SHIPPING
+		if (UE::GameFeatures::ShouldDebugGFP(StateProperties.PluginName))
+		{
+			static int32 I = 0; 
+			++I;
+		}
+#endif
 		if (StateProperties.GetPluginProtocol() != EGameFeaturePluginProtocol::InstallBundle)
 		{
 			bComplete = true;
@@ -3158,10 +3215,11 @@ struct FGameFeaturePluginState_AssetDependencyStreaming : public FGameFeaturePlu
 			if (InstallMode == EStreamingAssetInstallMode::GfpRequiredOnly)
 			{
 				InstallArgs.TagSets.Emplace(TEXTVIEW("required"));
+				InstallArgs.TagSets.Emplace(StateProperties.PluginName);
 			}
 			InstallArgs.Options |= UE::IoStore::EOnDemandInstallOptions::InstallSoftReferences;
 			InstallArgs.Options |= UE::IoStore::EOnDemandInstallOptions::CallbackOnGameThread;
-			if (Options.bDoNotDownload)
+			if (Options.bDoNotDownload && !UE::GameFeatures::CVarForceGFPIADDownload.GetValueOnGameThread())
 			{
 				InstallArgs.Options |= UE::IoStore::EOnDemandInstallOptions::DoNotDownload;
 			}
@@ -3170,12 +3228,14 @@ struct FGameFeaturePluginState_AssetDependencyStreaming : public FGameFeaturePlu
 				InstallArgs.Options |= UE::IoStore::EOnDemandInstallOptions::AllowMissingDependencies;
 			}
 
+			InstallArgs.DebugName = StateProperties.PluginName;
+
 			InstallArgs.ContentHandle = UE::GameFeatures::GFPSharedInstallTracker.AddOnDemandContentHandle(
 				InstallBundle, UE::GameFeatures::EGFPInstallLevel::AssetStream);
 
 			IoStoreOnDemandContext->InstallRequests.Add(IoStore->Install(MoveTemp(InstallArgs),
 				// On Complete
-				[this, LambdaOnDemandContext = IoStoreOnDemandContext](const UE::IoStore::FOnDemandInstallResult& OnDemandInstallResult)
+				[this, LambdaOnDemandContext = IoStoreOnDemandContext, DebugName = InstallArgs.DebugName](const UE::IoStore::FOnDemandInstallResult& OnDemandInstallResult)
 				{
 					if (!LambdaOnDemandContext->bStateValid)
 					{
@@ -3190,6 +3250,9 @@ struct FGameFeaturePluginState_AssetDependencyStreaming : public FGameFeaturePlu
 
 						ErrorCode.ReplaceCharInline(TEXT(' '), TEXT('_'), ESearchCase::CaseSensitive);
 						Result = GetErrorResult(TEXT("IoStoreOnDemand.OnComplete."), ErrorCode, ErrorMessage);
+						UE_LOGFMT(LogGameFeatures, Error, "Failed to install request {DebugName} due to error {Error}",
+							("DebugName", DebugName),
+							("Error", Result.GetError()));
 						TryCancelState();
 					}
 
diff --git a/Engine/Source/Programs/AutomationTool/Scripts/CopyBuildToStagingDirectory.Automation.cs b/Engine/Source/Programs/AutomationTool/Scripts/CopyBuildToStagingDirectory.Automation.cs
--- a/Engine/Source/Programs/AutomationTool/Scripts/CopyBuildToStagingDirectory.Automation.cs
+++ b/Engine/Source/Programs/AutomationTool/Scripts/CopyBuildToStagingDirectory.Automation.cs
@@ -2897,6 +2897,16 @@ namespace AutomationScripts
 			// List of allowed chunk names when generating on demand chunk(s) 
 			public List<Regex> OnDemandAllowedChunkFilters;
 
+			// Weather this rule should generate a completely new chunk 
+			public bool bGenerateNewChunk;
+			
+			// Filter assets out of GeneratedNewChunk which do not pass the AllowedChunkFilter (only used when bGenerateNewChunk = true)
+			public List<Regex> NewChunkAllowedChunkFilters;
+
+			// Generated chunk name (only valid for bGenerateNewChunk = true)
+			public string GeneratedChunkName;
+
+
 			public static bool IsMatch(PakFileRules PakRules, KeyValuePair<string, string> StagingFile)
 			{
 				bool bMatched = !PakRules.bDisabled &&
@@ -3057,6 +3067,9 @@ namespace AutomationScripts
 				string PakString;
 				PakRulesConfig.TryGetValue(SectionName, "OverridePaks", out PakString);
 
+				PakRulesConfig.TryGetValue(SectionName, "bGenerateNewChunk", out PakRules.bGenerateNewChunk);
+				PakRulesConfig.TryGetValue(SectionName, "GeneratedChunkName", out PakRules.GeneratedChunkName);
+
 				if (PakString != null && PakString.Length > 0)
 				{
 					PakRules.OverridePaks = PakString.Split(new char[] { ',', ' ' }, StringSplitOptions.RemoveEmptyEntries).ToList();
@@ -3079,7 +3092,7 @@ namespace AutomationScripts
 						continue;
 					}
 				}
-				else if (PakRules.OverridePaks == null && !PakRules.bOnDemand)
+				else if (PakRules.OverridePaks == null && !PakRules.bOnDemand && !PakRules.bGenerateNewChunk)
 				{
 					Logger.LogWarning("Error in PakFileRules {Arg0}, set to include but did not specify paks!", PakRules.Name);
 					continue;
@@ -3102,6 +3115,23 @@ namespace AutomationScripts
 					}
 				}
 
+				IReadOnlyList<string> NewChunkAllowedChunkFilters;
+				if (PakRulesConfig.TryGetValues(SectionName, "NewChunkAllowedChunkFilters", out NewChunkAllowedChunkFilters))
+				{
+					PakRules.NewChunkAllowedChunkFilters = new List<Regex>();
+					foreach (string Filter in NewChunkAllowedChunkFilters)
+					{
+						try
+						{
+							PakRules.NewChunkAllowedChunkFilters.Add(new Regex(Filter));
+						}
+						catch (System.Text.RegularExpressions.RegexParseException Ex)
+						{
+							Logger.LogError("Failed to parse allowed chunk filter, reason '{Arg1}'", Ex.Message);
+						}
+					}
+				}
+
 				IReadOnlyList<string> FilesEnumberable;
 				if (PakRulesConfig.TryGetValues(SectionName, "Files", out FilesEnumberable))
 				{
@@ -5566,6 +5596,30 @@ namespace AutomationScripts
 							}
 						}
 
+						if (MatchingRule.HasValue && MatchingRule.Value.bGenerateNewChunk)
+						{
+							if (!string.IsNullOrEmpty(MatchingRule.Value.GeneratedChunkName))
+							{
+								bool bMatchesFilter = true;
+								if (MatchingRule.Value.NewChunkAllowedChunkFilters != null)
+								{
+									bMatchesFilter = MatchingRule.Value.NewChunkAllowedChunkFilters.Any(Filter => Filter.IsMatch(TargetChunk.ChunkName));
+								}
+								if (bMatchesFilter)
+								{
+									string ChunkName = MatchingRule.Value.GeneratedChunkName;
+
+									AddChunkDefinition(AdditionalChunks, Chunk, ChunkName);
+									TargetChunk = AdditionalChunks[ChunkName];
+								}
+							}
+							else
+							{
+								Logger.LogError("Unable to generate new chunk for rule {Arg0}. GeneratedChunkName is not set.", MatchingRule.Value.Name);
+							}
+							
+						}
+
 						if (bStageLoose || bOnDemand)
 						{
 							// make a new separate pak if not already marked as on demand
