[Backout] - CL51593022
#rnx
Original CL Desc
-----------------------------------------------------------------
1) GFP_B was SetDest -> Active
2) GFP_B gets canceled back to -> Registered, we set bTryCancel=true and we are letting it trying to finish that transition
3) GFP_B has a parent GFP_A trying to Activte which requires its depends to now be Active, so it iterates through an tries to TransitionDepends on all its children

GFP_B has bTryCancel=True and it may be at Dest state which will fire a callback that it Completed a transition with an Error = Canceled

In these cases, we have canceled, the cancel in a sense, and should just try to move that GFP to the state the Parent requires by re-trying a TransitionDepends only once

diff --git a/Engine/Plugins/Runtime/GameFeatures/Source/GameFeatures/Private/GameFeaturePluginStateMachine.cpp b/Engine/Plugins/Runtime/GameFeatures/Source/GameFeatures/Private/GameFeaturePluginStateMachine.cpp
--- a/Engine/Plugins/Runtime/GameFeatures/Source/GameFeatures/Private/GameFeaturePluginStateMachine.cpp
+++ b/Engine/Plugins/Runtime/GameFeatures/Source/GameFeatures/Private/GameFeaturePluginStateMachine.cpp
@@ -1152,9 +1152,9 @@ struct FTransitionDependenciesGameFeaturePluginState : public FGameFeaturePlugin
 			}
 		}
 
-		for (FDepResult& Dep : RemainingDependencies)
+		for (FDepResultPair& Pair : RemainingDependencies)
 		{
-			UGameFeaturePluginStateMachine* RemainingDependency = Dep.Machine.Get();
+			UGameFeaturePluginStateMachine* RemainingDependency = Pair.Key.Get();
 			if (!RemainingDependency)
 			{
 				// One of the dependency state machines was destroyed before finishing
@@ -1162,14 +1162,14 @@ struct FTransitionDependenciesGameFeaturePluginState : public FGameFeaturePlugin
 				return;
 			}
 
-			if (Dep.Result.HasError())
+			if (Pair.Value.HasError())
 			{
 				const FStringView ShortUrl = StateProperties.PluginIdentifier.GetIdentifyingString();
 				const FStringView DepShortUrl = GameFeaturesSubsystem.GetPluginIdentifier(RemainingDependency).GetIdentifyingString();
 				UE_LOG(LogGameFeatures, Error, TEXT("Plugin (%.*s) dependency (%.*s) failed to transition with error %s"),
 					ShortUrl.Len(), ShortUrl.GetData(),
 					DepShortUrl.Len(), DepShortUrl.GetData(),
-					*Dep.Result.GetError());
+					*Pair.Value.GetError());
 				StateStatus.SetTransitionError(TransitionPolicy::GetErrorState(), GetErrorResult(TEXT("Failed_Dependency_Transition")));
 				if (GameFeaturesSubsystem.IsErrorStateUnrecoverable(RemainingDependency))
 				{
@@ -1244,14 +1244,14 @@ struct FTransitionDependenciesGameFeaturePluginState : public FGameFeaturePlugin
 			UGameFeaturePluginStateMachine* NewMachine = UGameFeaturesSubsystem::Get().FindOrCreateGameFeaturePluginStateMachine(GameFeaturesSubsystem.GetPluginURL(Dependency), DepProtocolOptions);
 			checkf(NewMachine != Dependency, TEXT("Game Feature Plugin %s should have already been removed from subsystem!"), *GameFeaturesSubsystem.GetPluginURL(Dependency));
 
-			const int32 Index = RemainingDependencies.IndexOfByPredicate([Dependency](const FDepResult& Dep)
+			const int32 Index = RemainingDependencies.IndexOfByPredicate([Dependency](const FDepResultPair& Pair)
 				{
-					return Dep.Machine == Dependency;
+					return Pair.Key == Dependency;
 				});
 
 			check(Index != INDEX_NONE);
-			FDepResult& FoundDep = RemainingDependencies[Index];
-			FoundDep.Machine = NewMachine;
+			FDepResultPair& FoundDep = RemainingDependencies[Index];
+			FoundDep.Key = NewMachine;
 
 			GameFeaturesSubsystem.RemovePendingTransitionCallback(Dependency, this);
 			GameFeaturesSubsystem.RemovePendingCancelCallback(Dependency, this);
@@ -1284,29 +1284,18 @@ struct FTransitionDependenciesGameFeaturePluginState : public FGameFeaturePlugin
 
 	void OnDependencyTransitionComplete(UGameFeaturePluginStateMachine* Dependency, const UE::GameFeatures::FResult& Result)
 	{
-		const int32 Index = RemainingDependencies.IndexOfByPredicate([Dependency](const FDepResult& Dep)
+		const int32 Index = RemainingDependencies.IndexOfByPredicate([Dependency](const FDepResultPair& Pair)
 			{
-				return Dep.Machine == Dependency;
+				return Pair.Key == Dependency;
 			});
 
 		if (ensure(Index != INDEX_NONE))
 		{
-			FDepResult& FoundDep = RemainingDependencies[Index];
+			FDepResultPair& FoundDep = RemainingDependencies[Index];
 
 			if (Result.HasError())
 			{
-				// If our error is a cancel, retry once — the dependency may have been canceled by a competing
-				// request and is now at a stable state we can transition from (e.g. Registered after a TryCancel
-				// settled there). Without this retry, we'd end up stuck in ErrorActivatingDependencies even
-				// though the dependency is perfectly healthy.
-				if (Result.GetError() == UE::GameFeatures::CanceledResult.GetError() && !FoundDep.bRetriedAfterCancel)
-				{
-					FoundDep.bRetriedAfterCancel = true;
-					TransitionDependency(Dependency);
-					return;
-				}
-
-				FoundDep.Result = Result;
+				FoundDep.Value = Result;
 			}
 			else
 			{
@@ -1343,9 +1332,9 @@ struct FTransitionDependenciesGameFeaturePluginState : public FGameFeaturePlugin
 		if (RemainingDependencies.Num() > 0)
 		{
 			UGameFeaturesSubsystem& GameFeaturesSubsystem = UGameFeaturesSubsystem::Get();
-			for (FDepResult& Dep : RemainingDependencies)
+			for (FDepResultPair& Pair : RemainingDependencies)
 			{
-				UGameFeaturePluginStateMachine* RemainingDependency = Dep.Machine.Get();
+				UGameFeaturePluginStateMachine* RemainingDependency = Pair.Key.Get();
 				if (RemainingDependency)
 				{
 					GameFeaturesSubsystem.RemovePendingTransitionCallback(RemainingDependency, this);
@@ -1363,16 +1352,8 @@ struct FTransitionDependenciesGameFeaturePluginState : public FGameFeaturePlugin
 		bRequestedDependencies = false;
 	}
 
-	struct FDepResult
-	{
-		FDepResult(UGameFeaturePluginStateMachine* InMachine, UE::GameFeatures::FResult InResult)
-			: Machine(InMachine), Result(MoveTemp(InResult)) {}
-
-		TWeakObjectPtr<UGameFeaturePluginStateMachine> Machine;
-		UE::GameFeatures::FResult Result;
-		bool bRetriedAfterCancel = false;
-	};
-	TArray<FDepResult> RemainingDependencies;
+	using FDepResultPair = TPair<TWeakObjectPtr<UGameFeaturePluginStateMachine>, UE::GameFeatures::FResult>;
+	TArray<FDepResultPair> RemainingDependencies;
 	bool bRequestedDependencies = false;
 	bool bCheckedRealtimeMode = false;
 };
