[GFP] Remove CVar "GameFeaturePlugin.AllowForceMonolithicShaderLibrary" to reduce logspam since loading chunked shader libraries from plugins is not supported.

This also removes `bForceMonolithicShaderLibrary` which is passed to `FShaderCodeLibrary::OpenPluginShaderLibrary()` where it is always ignored.

#rnx
#Synced-CL 51753475

diff --git a/Engine/Plugins/Runtime/GameFeatures/Source/GameFeatures/Private/GameFeaturePluginStateMachine.cpp b/Engine/Plugins/Runtime/GameFeatures/Source/GameFeatures/Private/GameFeaturePluginStateMachine.cpp
--- a/Engine/Plugins/Runtime/GameFeatures/Source/GameFeatures/Private/GameFeaturePluginStateMachine.cpp
+++ b/Engine/Plugins/Runtime/GameFeatures/Source/GameFeatures/Private/GameFeaturePluginStateMachine.cpp
@@ -133,10 +133,6 @@ namespace UE::GameFeatures
 		true,
 		TEXT("Enable to use async loading as well async downloading and registering"));
 
-	static TAutoConsoleVariable<bool> CVarAllowForceMonolithicShaderLibrary(TEXT("GameFeaturePlugin.AllowForceMonolithicShaderLibrary"),
-		true,
-		TEXT("Enable to force only searching for monolithic shader libs when possible"));
-
 	static TAutoConsoleVariable<bool> CVarForceSyncRegisterStartupPlugins(TEXT("GameFeaturePlugin.ForceSyncRegisterStartupPlugins"),
 		true,
 		TEXT("If true, all plugins loaded during startup will be synchronously registered to ensure things are initialized in time, this only applies if AsyncLoad is enabled"));
@@ -2520,7 +2516,6 @@ struct FGameFeaturePluginState_Mounting : public FGameFeaturePluginState
 	ESubState StartedSubStates = ESubState::None;
 	ESubState CompletedSubStates = ESubState::None;
 	bool bCheckedRealtimeMode = false;
-	bool bForceMonolithicShaderLibrary = true;	// use monolithic unless a DLC plugin is chunked
 
 	static UE::Tasks::FPipe ShaderlibPipe;
 
@@ -2545,11 +2540,6 @@ struct FGameFeaturePluginState_Mounting : public FGameFeaturePluginState
 			}
 		}
 
-		if (bForceMonolithicShaderLibrary && BundleResult.bContainsChunks)
-		{
-			bForceMonolithicShaderLibrary = false;
-		}
-
 		if (PendingBundles.IsEmpty())
 		{
 			IInstallBundleManager::InstallBundleCompleteDelegate.RemoveAll(this);
@@ -2615,16 +2605,12 @@ struct FGameFeaturePluginState_Mounting : public FGameFeaturePluginState
 		StartedSubStates = ESubState::None;
 		CompletedSubStates = ESubState::None;
 		bCheckedRealtimeMode = false;
-		bForceMonolithicShaderLibrary = false;
 
 		if (StateProperties.GetPluginProtocol() != EGameFeaturePluginProtocol::InstallBundle)
 		{
 			return;
 		}
 
-		// Assume monolithic shader, will be set to false if chunks are detected
-		bForceMonolithicShaderLibrary = UE::GameFeatures::CVarAllowForceMonolithicShaderLibrary.GetValueOnGameThread();
-		
 		TSharedPtr<IInstallBundleManager> BundleManager = IInstallBundleManager::GetPlatformInstallBundleManager();
 
 		const FInstallBundlePluginProtocolMetaData& MetaData = StateProperties.ProtocolMetadata.GetSubtype<FInstallBundlePluginProtocolMetaData>();
@@ -2670,14 +2656,6 @@ struct FGameFeaturePluginState_Mounting : public FGameFeaturePluginState
 				PakFileMountedDelegateHandle = FCoreDelegates::GetOnPakFileMounted2().AddRaw(this, &FGameFeaturePluginState_Mounting::OnPakFileMounted);
 			}
 		}
-
-		for (const FInstallBundleRequestResultInfo& BundleResult : RequestInfo.BundleResults)
-		{
-			if (bForceMonolithicShaderLibrary && BundleResult.bContainsChunks)
-			{
-				bForceMonolithicShaderLibrary = false;
-			}
-		}
 	}
 
 	void UpdateState_MountPlugin(bool bLoadPluginIniHierarchy)
@@ -2763,7 +2741,7 @@ struct FGameFeaturePluginState_Mounting : public FGameFeaturePluginState
 			if (bManuallyOpenPluginShaderLibrary)
 			{
 				TSharedPtr<IPlugin> Plugin = IPluginManager::Get().FindPlugin(StateProperties.PluginName);
-				FShaderCodeLibrary::OpenPluginShaderLibrary(*Plugin, bForceMonolithicShaderLibrary);
+				FShaderCodeLibrary::OpenPluginShaderLibrary(*Plugin);
 			}
 			CompletedSubStates |= ESubState::MountPlugin;
 			return;
@@ -2783,7 +2761,7 @@ struct FGameFeaturePluginState_Mounting : public FGameFeaturePluginState
 			// TEMP HACK - use a pipe because if this goes too wide we can end up blocking all available tasks.
 			ShaderlibPipe.Launch(UE_SOURCE_LOCATION, [this, Plugin]
 			{
-				FShaderCodeLibrary::OpenPluginShaderLibrary(*Plugin, bForceMonolithicShaderLibrary);
+				FShaderCodeLibrary::OpenPluginShaderLibrary(*Plugin);
 
 				ExecuteOnGameThread(UE_SOURCE_LOCATION, [this]
 				{
