[GFP] Improve unmounting performance by batching primary asset cleanup across multiple game feature plugins
Introduce a suspension mechanism that defers primary asset unloading until all plugins in a bulk transition are processed, reducing redundant asset registry traversals
Add a new SuspendPrimaryAssetsCleanup transition option integrated into the existing bulk transition context

#rnx
#Synced-CL 51794024

diff --git a/Engine/Plugins/Runtime/GameFeatures/Source/GameFeatures/Private/GameFeaturesSubsystem.cpp b/Engine/Plugins/Runtime/GameFeatures/Source/GameFeatures/Private/GameFeaturesSubsystem.cpp
--- a/Engine/Plugins/Runtime/GameFeatures/Source/GameFeatures/Private/GameFeaturesSubsystem.cpp
+++ b/Engine/Plugins/Runtime/GameFeatures/Source/GameFeatures/Private/GameFeaturesSubsystem.cpp
@@ -5,11 +5,14 @@
 #include "Algo/AllOf.h"
 #include "Algo/AnyOf.h"
 #include "Algo/TopologicalSort.h"
+#include "Algo/Transform.h"
 #include "Algo/Unique.h"
 #include "AssetRegistry/AssetData.h"
 #include "AssetRegistry/IAssetRegistry.h"
+#include "Async/Mutex.h"
 #include "Async/ParallelFor.h"
 #include "BundlePrereqCombinedStatusHelper.h"
+#include "Containers/Array.h"
 #include "Experimental/UnifiedError/UnifiedError.h"
 #include "Dom/JsonValue.h"
 #include "GameFeaturePluginOperationResult.h"
@@ -25,12 +28,15 @@
 #include "IO/IoStoreOnDemand.h"
 #include "JsonUtils/JsonConversion.h"
 #include "JsonUtils/RapidJsonUtils.h"
+#include "Logging/LogVerbosity.h"
 #include "Logging/StructuredLog.h"
 #include "Misc/App.h"
+#include "Misc/EnumClassFlags.h"
 #include "Misc/FileHelper.h"
 #include "Misc/PackageName.h"
 #include "Misc/Paths.h"
 #include "Misc/PathViews.h"
+#include "Misc/ScopeLock.h"
 #include "Misc/ScopeRWLock.h"
 #include "Serialization/JsonSerializer.h"
 #include "Stats/StatsMisc.h"
@@ -311,6 +317,41 @@
 #else
 	class FPackageLoadTracker {};
 #endif
+
+	int32 SuspendPrimaryAssetCleanupCounter = 0;
+	TMultiMap<FString, FString> ScheduledAssetScanPathsByPluginName;
+	void CleanupPrimaryAssetByScanPaths(TConstArrayView<FName> ScanPaths)
+	{
+		if (!ScanPaths.IsEmpty())
+		{
+			IAssetRegistry* AssetRegistry = IAssetRegistry::Get();
+			UAssetManager* AssetManager = UAssetManager::GetIfInitialized();
+			if (AssetRegistry && AssetManager)
+			{
+				TRACE_CPUPROFILER_EVENT_SCOPE(GFP_RemoveGameFeatureFromAssetManager_TraverseAR);
+				FARFilter AssetsFilter;
+				AssetsFilter.PackagePaths = ScanPaths;
+				AssetsFilter.bRecursivePaths = true;
+				AssetsFilter.TagsAndValues.Add(FPrimaryAssetId::PrimaryAssetTypeTag);
+				AssetsFilter.RecursiveClassPathsExclusionSet.Add(UGameFeatureData::StaticClass()->GetClassPathName());
+				AssetsFilter.bRecursiveClasses = true;
+
+				TArray<FAssetData> AssetDatas;
+				if (AssetRegistry->GetInMemoryAssets(AssetsFilter, AssetDatas))
+				{
+					TArray<FPrimaryAssetId> PotentialAssets;
+					PotentialAssets.Reserve(AssetDatas.Num());
+
+					for (const FAssetData& Asset : AssetDatas)
+					{
+						PotentialAssets.Add(Asset.GetPrimaryAssetId());
+					}
+
+					AssetManager->UnloadPrimaryAssets(PotentialAssets);
+				}
+			}
+		}
+	}
 }
 
 TSet<FString> UGameFeaturesSubsystem::DlcPlugins;
@@ -909,6 +950,12 @@
 	}
 	LocalAssetManager.PopBulkScanning();
 
+	if (UE::GameFeatures::CVarUnloadPrimaryAssetsOnUnregister->GetBool() && 
+		UE::GameFeatures::SuspendPrimaryAssetCleanupCounter > 0)
+	{
+		UE::GameFeatures::ScheduledAssetScanPathsByPluginName.Remove(PluginName);
+	}
+
 	const UAssetManagerSettings& Settings = LocalAssetManager.GetSettings();
 	for (const FPrimaryAssetRulesCustomOverride& Override : Settings.CustomPrimaryAssetRules)
 	{
@@ -925,64 +972,10 @@
 	FString PluginRootPath = TEXT("/") + PluginName + TEXT("/");
 	UAssetManager& LocalAssetManager = UAssetManager::Get();
 
-	// Cleanup primary assets from asset manager if they were loaded by users directly.
-	if (UE::GameFeatures::CVarUnloadPrimaryAssetsOnUnregister->GetBool())
-	{
-		TArray<FString> ScanPaths;
-		if (GameFeatureToRemove->GetPrimaryAssetTypesToScan().Num() > 1)
-		{
-			int32 TotalSizeRequired = 0;
-			for (const FPrimaryAssetTypeInfo& Type : GameFeatureToRemove->GetPrimaryAssetTypesToScan())
-			{
-				TotalSizeRequired += Type.AssetScanPaths.Num();
-			}
-			ScanPaths.Reserve(TotalSizeRequired);
-		}
-		for (const FPrimaryAssetTypeInfo& Type : GameFeatureToRemove->GetPrimaryAssetTypesToScan())
-		{
-			ScanPaths.Append(Type.AssetScanPaths);
-		}
-
-		if (!ScanPaths.IsEmpty())
-		{
-			if (IAssetRegistry* AR = IAssetRegistry::Get())
-			{
-				FARFilter AssetsFilter;
-				for (const FString& Path : ScanPaths)
-				{
-					AssetsFilter.PackagePaths.Add(FName(Path));
-				}
-				AssetsFilter.bRecursivePaths = true;
-				AssetsFilter.TagsAndValues.Add(FPrimaryAssetId::PrimaryAssetTypeTag);
-				AssetsFilter.RecursiveClassPathsExclusionSet.Add(UGameFeatureData::StaticClass()->GetClassPathName());
-				AssetsFilter.bRecursiveClasses = true;
-
-				TArray<FAssetData> AssetDatas;
-				if (AR->GetInMemoryAssets(AssetsFilter, AssetDatas))
-				{
-					TArray<FPrimaryAssetId> PotentialAssets;
-					PotentialAssets.Reserve(AssetDatas.Num());
-
-					for (const FAssetData& Asset : AssetDatas)
-					{
-						PotentialAssets.Add(Asset.GetPrimaryAssetId());
-					}
-
-					UAssetManager::Get().UnloadPrimaryAssets(PotentialAssets);
-				}
-			}
-		}
-	}
-
 	// Remove scan paths from asset manager
+	TArray<FString> ScanPaths;
 	for (FPrimaryAssetTypeInfo TypeInfo : GameFeatureToRemove->GetPrimaryAssetTypesToScan())
 	{
-		if (AddedPrimaryAssetTypes.Contains(TypeInfo.PrimaryAssetType))
-		{
-			LocalAssetManager.RemovePrimaryAssetType(TypeInfo.PrimaryAssetType);
-			continue;
-		}
-
 		for (FDirectoryPath& Path : TypeInfo.GetDirectories())
 		{
 			FixPluginPackagePath(Path.Path, PluginRootPath, false);
@@ -994,7 +987,38 @@
 			continue;
 		}
 
+		// We need to collect all possible scan paths before we remove them from asset manager
+		if (UE::GameFeatures::CVarUnloadPrimaryAssetsOnUnregister->GetBool())
+		{
+			ScanPaths.Append(TypeInfo.AssetScanPaths);
+		}
+
+		if (AddedPrimaryAssetTypes.Contains(TypeInfo.PrimaryAssetType))
+		{
+			LocalAssetManager.RemovePrimaryAssetType(TypeInfo.PrimaryAssetType);
+			continue;
+		}
+
 		LocalAssetManager.RemoveScanPathsForPrimaryAssets(TypeInfo.PrimaryAssetType, TypeInfo.AssetScanPaths, TypeInfo.AssetBaseClassLoaded, TypeInfo.bHasBlueprintClasses, TypeInfo.bIsEditorOnly);
 	}
+
+	if (UE::GameFeatures::CVarUnloadPrimaryAssetsOnUnregister->GetBool())
+	{
+		if (UE::GameFeatures::SuspendPrimaryAssetCleanupCounter > 0)
+		{
+			UE::GameFeatures::ScheduledAssetScanPathsByPluginName.Reserve(UE::GameFeatures::ScheduledAssetScanPathsByPluginName.Num() + ScanPaths.Num());
+			for (const FString& Path : ScanPaths)
+			{
+				UE::GameFeatures::ScheduledAssetScanPathsByPluginName.Add(PluginName, Path);
+			}
+		}
+		else
+		{
+			TArray<FName> ScanPathNames;
+			ScanPathNames.Reserve(ScanPaths.Num());
+			Algo::Transform(ScanPaths, ScanPathNames, [](const FString& Path) { return FName(Path); });
+			UE::GameFeatures::CleanupPrimaryAssetByScanPaths(ScanPathNames);
+		}
+	}
 }
 
@@ -2251,6 +2275,30 @@
 		bool bWasSuspended = false;
 	};
 
+	struct FSuspendPrimaryAssetCleanup
+	{
+		FSuspendPrimaryAssetCleanup()
+		{
+			++UE::GameFeatures::SuspendPrimaryAssetCleanupCounter;
+		}
+
+		~FSuspendPrimaryAssetCleanup()
+		{
+			ensure(UE::GameFeatures::SuspendPrimaryAssetCleanupCounter > 0);
+			if (--UE::GameFeatures::SuspendPrimaryAssetCleanupCounter == 0)
+			{
+				TArray<FName> ScanPaths;
+				ScanPaths.Reserve(UE::GameFeatures::ScheduledAssetScanPathsByPluginName.Num());
+				for (const TTuple<FString, FString>& Pair : UE::GameFeatures::ScheduledAssetScanPathsByPluginName)
+				{
+					ScanPaths.Add(FName(Pair.Value));
+				}
+				UE::GameFeatures::ScheduledAssetScanPathsByPluginName.Empty();
+				UE::GameFeatures::CleanupPrimaryAssetByScanPaths(ScanPaths);
+			}
+		}
+	};
+
 	enum class ETransitionOptions
 	{
 		None							= 0,
@@ -2259,7 +2307,8 @@
 		SuspendAssetRegistry			= 1 << 2,
 		NotifyObservers					= 1 << 3,
 		SuspendRenderAssetStreamer		= 1 << 4,
-		All = SuspendPluginManager | SuspendGameplayTagTreeBroadcast | SuspendAssetRegistry | NotifyObservers | SuspendRenderAssetStreamer,
+		SuspendPrimaryAssetsCleanup		= 1 << 5,
+		All = SuspendPluginManager | SuspendGameplayTagTreeBroadcast | SuspendAssetRegistry | NotifyObservers | SuspendRenderAssetStreamer | SuspendPrimaryAssetsCleanup,
 	};
 	ENUM_CLASS_FLAGS(ETransitionOptions)
 
@@ -2273,6 +2322,7 @@
 		TOptional<FSuspendPluginManager> SuspendPluginManager;
 		TOptional<FSuspendGameplayTagTreeChangedBroadcast> SuspendGameplayTagTreeBroadcast;
 		TOptional<IAssetRegistry::FSuspendPathsDismounting> SuspendAssetRegistry;
+		TOptional<FSuspendPrimaryAssetCleanup> SuspendPrimaryAssetsCleanup;
 		TOptional<FNotifyObservers> NotifyObservers;
 
 		UE_NODISCARD_CTOR FTransitionContenxt(ETransitionOptions TransitionOptions = ETransitionOptions::All)
@@ -2303,6 +2353,11 @@
 				SuspendAssetRegistry.Emplace();
 			}
 
+			if (EnumHasAnyFlags(TransitionOptions, ETransitionOptions::SuspendPrimaryAssetsCleanup))
+			{
+				SuspendPrimaryAssetsCleanup.Emplace();
+			}
+
 			if (EnumHasAnyFlags(TransitionOptions, ETransitionOptions::NotifyObservers))
 			{
 				NotifyObservers.Emplace();
@@ -2318,6 +2373,7 @@
 			// And by that point all other resources need to be properly cleared,
 			// i.e. other batches need to finish.
 			NotifyObservers.Reset(); // Some observers are responsible for clearing references; must happen before SuspendPluginManager.Reset()
+			SuspendPrimaryAssetsCleanup.Reset(); // assets cleanup requires data in asset registry still present
 			SuspendAssetRegistry.Reset();
 			SuspendGameplayTagTreeBroadcast.Reset();
 			SuspendPluginManager.Reset();
@@ -2404,6 +2460,8 @@
 		EnumRemoveFlags(TransitionOptions, GameFeaturesSubsystem::ETransitionOptions::SuspendRenderAssetStreamer);
 	}
 
+	EnumAddFlags(TransitionOptions, GameFeaturesSubsystem::ETransitionOptions::SuspendPrimaryAssetsCleanup);
+
 	TSharedRef<GameFeaturesSubsystem::FTransitionContenxt> TransitionContext = MakeShared<GameFeaturesSubsystem::FTransitionContenxt>(TransitionOptions);
 	TransitionContext->CompleteDelegate = CompleteDelegate;
 
