Enable new executor in GFP state machines. Group execution by state and defer execution in LoadBuiltIn until all state machines are created.

#Synced-CL 51794024

#Synced-CL 51794024

#Synced-CL 51794024

diff --git a/Engine/Plugins/Runtime/GameFeatures/Source/GameFeatures/Private/GameFeaturePluginStateMachine.cpp b/Engine/Plugins/Runtime/GameFeatures/Source/GameFeatures/Private/GameFeaturePluginStateMachine.cpp
--- a/Engine/Plugins/Runtime/GameFeatures/Source/GameFeatures/Private/GameFeaturePluginStateMachine.cpp
+++ b/Engine/Plugins/Runtime/GameFeatures/Source/GameFeatures/Private/GameFeaturePluginStateMachine.cpp
@@ -77,7 +77,7 @@ namespace UE::GameFeatures
 
 	static UE::GameFeatures::FResult CanceledResult = MakeError(StateMachineErrorNamespace + TEXT("Canceled"));
 
-	bool bUseNewExecutor = false;
+	bool bUseNewExecutor = true;
 	static FAutoConsoleVariableRef CVarUseNewExecutor(TEXT("GameFeaturePlugin.UseNewExecutor"),
 		bUseNewExecutor,
 		TEXT("Should control of state machine updates be handled by the GameFeaturesSubsystem (new behavior) to all more batch & parallel processing."),
@@ -87,7 +87,7 @@ namespace UE::GameFeatures
 	struct FExecutionOptions
 	{
 		// Group by current state before executing
-		bool bGrouped = false;
+		bool bGrouped = true;
 		// Execute the update requests (potentially within each group) in a random order
 		bool bShuffle = false;
 		// Reverse the order of update requests so that deferred request are executed in LIFO order
@@ -107,7 +107,7 @@ namespace UE::GameFeatures
 		TEXT("Reverse deferred updates before execution to execute in LIFO order. For testing deferred dependencies."),
 		ECVF_Default);
 
-	bool bLoadBuiltin_DeferExecution = false;
+	bool bLoadBuiltin_DeferExecution = true;
 	static FAutoConsoleVariableRef CVarLoadBuiltin_DeferExecution (TEXT("GameFeaturePlugin.LoadBuiltins.DeferExecution"),
 		bLoadBuiltin_DeferExecution ,
 		TEXT("When GameFeaturePlugin.UseNewExecutor is true, should execution of state machine updates in LoadBuiltInGameFeaturePlugins be deferred until all state transitions have been queued?"),
@@ -283,7 +283,9 @@ namespace UE::GameFeatures
 		return Priority;
 	}
 
-	static bool bRealtimeMode = false;
+	bool bRealtimeMode = false;
+	constexpr double RealtimeMaxFrameTime = 0.033; // 30fps
+	constexpr double RealtimeAllottedTime = RealtimeMaxFrameTime / 2;
 
 	class FRealtimeMode : public TSharedFromThis<FRealtimeMode>
 	{
@@ -324,8 +326,6 @@ namespace UE::GameFeatures
 			TSharedRef<FRealtimeMode> SelfRef = AsShared();
 
 			{
-				constexpr double MaxFrameTime = 0.033; // 30fps
-				constexpr double AllottedTime = MaxFrameTime / 2;
 				const double StartTime = FPlatformTime::Seconds();
 
 				UGameFeaturesSubsystem& Subsystem = UGameFeaturesSubsystem::Get();
@@ -335,7 +335,7 @@ namespace UE::GameFeatures
 					Subsystem.TryUpdateStateMachineImmediate(UpdateRequest);
 
 					const double ElapsedTime = FPlatformTime::Seconds() - StartTime;
-					if ((ElapsedTime > AllottedTime) || ((DeltaTime + ElapsedTime) > MaxFrameTime))
+					if ((ElapsedTime > RealtimeAllottedTime) || ((DeltaTime + ElapsedTime) > RealtimeMaxFrameTime))
 					{
 						break;
 					}
@@ -368,6 +368,7 @@ namespace UE::GameFeatures
 			{
 				if (bUseNewExecutor)
 				{
+					// new executor reads bRealtimeMode directly and enforces timings
 					return;
 				}
 				if (Var->GetBool())
@@ -4283,12 +4284,14 @@ void UE::GameFeatures::FStateMachineExecutor::TryResumeExecution(bool bStayPause
 	if (!bPauseImmediateExecution)
 	{
 		TRACE_CPUPROFILER_EVENT_SCOPE_STR("FStateMachineExecutor::TryResumeExecution")
-		UpdateStateMachines();
+		// No time limit based on delta time here, just use the per-frame budget
+		UpdateStateMachines(0.0f); 
 	}
 }
 
-bool UE::GameFeatures::FStateMachineExecutor::UpdateStateMachines(int32 MaxIterations)
+bool UE::GameFeatures::FStateMachineExecutor::UpdateStateMachines(float DeltaTime, int32 MaxIterations)
 {
+	const double StartTime = UE::GameFeatures::bRealtimeMode ? FPlatformTime::Seconds() : 0.0;
 	UAssetManager& LocalAssetManager = UAssetManager::Get();
 	LocalAssetManager.PushBulkScanning();
 	ON_SCOPE_EXIT 
@@ -4343,8 +4346,23 @@ bool UE::GameFeatures::FStateMachineExecutor::UpdateStateMachines(int32 MaxItera
 				break;
 			}
 
-			UpdateStateMachinesGrouped(State, MakeArrayView(Queues[State]));
+			bool bTimeLimitReached = UpdateStateMachinesGrouped(StartTime, DeltaTime, State, MakeArrayView(Queues[State]));
 			Queues[State].Reset();
+			if (bTimeLimitReached)
+			{
+				OrderedUpdateQueue.Reserve(OrderedUpdateQueue.Num() + Algo::TransformAccumulate(Queues, UE_PROJECTION(GetNum), 0));
+				// UpdateStateMachinesGrouped will already have requeued Queues[State], all other queues need to be requeued
+				// They can be queued at the end because we'll sort again next time
+				for (EGameFeaturePluginState It : TEnumRange<EGameFeaturePluginState>())
+				{
+					Algo::Transform(Queues[It], OrderedUpdateQueue, [](UGameFeaturePluginStateMachine* StateMachine)
+					{
+						StateMachine->StateProperties.bQueuedForUpdate = true;
+						return StateMachine->StateProperties.Id;
+					});
+				}
+				break;
+			}
 		}
 	}
 	else 
@@ -4372,38 +4390,71 @@ bool UE::GameFeatures::FStateMachineExecutor::UpdateStateMachines(int32 MaxItera
 		}
 		OrderedUpdateQueue.Reset();
 
+		UpdateStateMachinesRealtime(StartTime, DeltaTime, LocalUpdateQueue);
+	}
+
+	// Return if this function should be called again next frame
+	return OrderedUpdateQueue.Num() > 0;
+}
+
+bool UE::GameFeatures::FStateMachineExecutor::UpdateStateMachinesRealtime(double StartTime, float DeltaTime, TConstArrayView<UGameFeaturePluginStateMachine*> LocalUpdateQueue)
+{
+	if (bRealtimeMode)
+	{
+		for (int32 i=0; i < LocalUpdateQueue.Num(); ++i)
+		{
+			UGameFeaturePluginStateMachine* StateMachine = LocalUpdateQueue[i];
+			OwningSubsystem->UpdateStateMachine(StateMachine, {});
+
+			const double ElapsedTime = FPlatformTime::Seconds() - StartTime;
+			if ((ElapsedTime > RealtimeAllottedTime) || ((DeltaTime + ElapsedTime) > RealtimeMaxFrameTime))
+			{
+				// Put anything we didn't finish back at the head of the queue
+				//-1 because we updated the machine at index i
+				TConstArrayView<UGameFeaturePluginStateMachine*> Remaining = LocalUpdateQueue.RightChop(i+1);
+				OrderedUpdateQueue.InsertZeroed(0, Remaining.Num());
+				for (int32 j=0; j < Remaining.Num(); ++j)
+				{
+					UGameFeaturePluginStateMachine* StateMachineToRequeue = Remaining[j];
+					StateMachineToRequeue->StateProperties.bQueuedForUpdate = true;
+					OrderedUpdateQueue[j] = StateMachineToRequeue->StateProperties.Id;
+				}
+
+				return true;
+			}
+		}
+	}
+	else
+	{ 
 		for (UGameFeaturePluginStateMachine* StateMachine : LocalUpdateQueue)
 		{
 			OwningSubsystem->UpdateStateMachine(StateMachine, {});
 		}
 	}
-	return OrderedUpdateQueue.Num() > 0;
-}
-
-#define GFP_STATE_OP(STATE_NAME, TXT) \
-template<> \
-void UE::GameFeatures::FStateMachineExecutor::UpdateStateMachinesGrouped<EGameFeaturePluginState::STATE_NAME>(TConstArrayView<UGameFeaturePluginStateMachine*> StateMachines) \
-{ \
-	TRACE_CPUPROFILER_EVENT_SCOPE_STR("UpdateStateMachinesGrouped<" #STATE_NAME ">"); \
-	for (UGameFeaturePluginStateMachine* StateMachine : StateMachines) \
-	{ \
-		OwningSubsystem->UpdateStateMachine(StateMachine, { 1 }); \
-	} \
+	return false;
 }
-	GAME_FEATURE_PLUGIN_STATE_LIST(GFP_STATE_OP)
-#undef GFP_STATE_OP
 
-void UE::GameFeatures::FStateMachineExecutor::UpdateStateMachinesGrouped(EGameFeaturePluginState State, TConstArrayView<UGameFeaturePluginStateMachine*> StateMachines)
+bool UE::GameFeatures::FStateMachineExecutor::UpdateStateMachinesGrouped(double StartTime, float DeltaTime, EGameFeaturePluginState State, TConstArrayView<UGameFeaturePluginStateMachine*> StateMachines)
 {
 	switch (State)
 	{
 #define GFP_STATE_OP(STATE_NAME, TXT) \
 	case EGameFeaturePluginState::STATE_NAME: \
-		UpdateStateMachinesGrouped<EGameFeaturePluginState::STATE_NAME>(StateMachines);\
+		{ \
+			TRACE_CPUPROFILER_EVENT_SCOPE_STR("UpdateStateMachinesGrouped<" #STATE_NAME ">"); \
+			return UpdateStateMachinesGrouped<EGameFeaturePluginState::STATE_NAME>(StartTime, DeltaTime, StateMachines);\
+		} \
 		break;
 	GAME_FEATURE_PLUGIN_STATE_LIST(GFP_STATE_OP)
 #undef GFP_STATE_OP
 	}
+	return false;
+}
+
+template<EGameFeaturePluginState State>
+bool UE::GameFeatures::FStateMachineExecutor::UpdateStateMachinesGrouped(double StartTime, float DeltaTime, TConstArrayView<UGameFeaturePluginStateMachine*> StateMachines)
+{
+	return UpdateStateMachinesRealtime(StartTime, DeltaTime, StateMachines);
 }
 
 void UE::GameFeatures::FStateMachineExecutor::QueueUpdateStateMachine(FGameFeaturePluginStateMachineProperties& InProperties)
diff --git a/Engine/Plugins/Runtime/GameFeatures/Source/GameFeatures/Private/GameFeaturePluginStateMachine.h b/Engine/Plugins/Runtime/GameFeatures/Source/GameFeatures/Private/GameFeaturePluginStateMachine.h
--- a/Engine/Plugins/Runtime/GameFeatures/Source/GameFeatures/Private/GameFeaturePluginStateMachine.h
+++ b/Engine/Plugins/Runtime/GameFeatures/Source/GameFeatures/Private/GameFeaturePluginStateMachine.h
@@ -191,10 +191,18 @@ namespace UE::GameFeatures
 		 */
 		bool bPauseImmediateExecution = false;
 
-		/** Update a set of state machines which are expected to be in the same state */
-		void UpdateStateMachinesGrouped(EGameFeaturePluginState State, TConstArrayView<UGameFeaturePluginStateMachine*> StateMachines);
+		/** 
+		 * Update a set of state machines which are expected to be in the same state. 
+		 * Returns true in realtime mode if time limit was exhausted.
+		 */
+		bool UpdateStateMachinesGrouped(double StartTime, float DeltaTime, EGameFeaturePluginState State, TConstArrayView<UGameFeaturePluginStateMachine*> StateMachines);
+		
+		// Templated to allow custom implementation per state 
 		template<EGameFeaturePluginState State>
-		void UpdateStateMachinesGrouped(TConstArrayView<UGameFeaturePluginStateMachine*> StateMachines);
+		bool UpdateStateMachinesGrouped(double StartTime, float DeltaTime, TConstArrayView<UGameFeaturePluginStateMachine*> StateMachines);
+
+		// Helper to update a list of state machines, requeing if we run out of time 
+		bool UpdateStateMachinesRealtime(double StartTime, float DeltaTime, TConstArrayView<UGameFeaturePluginStateMachine*> StateMachines);
 
 	public:
 		FStateMachineExecutor(UGameFeaturesSubsystem* InOwningSubsystem);
@@ -202,7 +210,7 @@ namespace UE::GameFeatures
 		UE_NONCOPYABLE(FStateMachineExecutor);
 
 		/** Update all state machines which are not at their target state and are not waiting on async work to complete */
-		bool UpdateStateMachines(int32 MaxIterations = 10000);	
+		bool UpdateStateMachines(float DeltaTime, int32 MaxIterations = 10000);	
 
 		/** 
 		 * Pause immediate execution of state machine updates and return whether it was already paused.
diff --git a/Engine/Plugins/Runtime/GameFeatures/Source/GameFeatures/Private/GameFeaturesSubsystem.cpp b/Engine/Plugins/Runtime/GameFeatures/Source/GameFeatures/Private/GameFeaturesSubsystem.cpp
--- a/Engine/Plugins/Runtime/GameFeatures/Source/GameFeatures/Private/GameFeaturesSubsystem.cpp
+++ b/Engine/Plugins/Runtime/GameFeatures/Source/GameFeatures/Private/GameFeaturesSubsystem.cpp
@@ -5015,7 +5015,7 @@ bool UGameFeaturesSubsystem::Tick(float DeltaTime)
 	bool bContinueTicking = TickBatchProcessing();
 	if (StateMachineExecutor)
 	{
-		bContinueTicking = StateMachineExecutor->UpdateStateMachines() || bContinueTicking;
+		bContinueTicking = StateMachineExecutor->UpdateStateMachines(DeltaTime) || bContinueTicking;
 	}
 	
 	if (!bContinueTicking)
