[GFP] Fix use after free in the FGFPSharedInstallTracker

#rnx
#Synced-CL 51996919

diff --git a/Engine/Plugins/Runtime/GameFeatures/Source/GameFeatures/Private/GameFeaturePluginStateMachine.cpp b/Engine/Plugins/Runtime/GameFeatures/Source/GameFeatures/Private/GameFeaturePluginStateMachine.cpp
--- a/Engine/Plugins/Runtime/GameFeatures/Source/GameFeatures/Private/GameFeaturePluginStateMachine.cpp
+++ b/Engine/Plugins/Runtime/GameFeatures/Source/GameFeatures/Private/GameFeaturePluginStateMachine.cpp
@@ -440,7 +440,7 @@ namespace UE::GameFeatures
 	struct FPluginInstallBundleReferencers
 	{
 		// GFPs using an install bundle and the relevant state of that GFP
-		TMap<FStringView, EGFPInstallLevel> GFPs;
+		TMap<FString, EGFPInstallLevel> GFPs;
 	};
 
 	struct FPluginIoStoreOnDemandHandles
@@ -523,7 +523,7 @@ namespace UE::GameFeatures
 						}
 					}
 
-					for (TPair<FStringView, EGFPInstallLevel>& GFPRef : PluginRefs->GFPs)
+					for (TPair<FString, EGFPInstallLevel>& GFPRef : PluginRefs->GFPs)
 					{
 						if (GFPRef.Value >= InLevel)
 						{
