Fix for incorrect management of install bundles from GFPs caused by refactor of state transitions to remove "skipped" states in CL 50263170. Bundles need to be released by the Releasing state and not just uninstalled by the Uninstalling state.

diff --git a/Engine/Plugins/Runtime/GameFeatures/Source/GameFeatures/Private/GameFeaturePluginStateMachine.cpp b/Engine/Plugins/Runtime/GameFeatures/Source/GameFeatures/Private/GameFeaturePluginStateMachine.cpp
--- a/Engine/Plugins/Runtime/GameFeatures/Source/GameFeatures/Private/GameFeaturePluginStateMachine.cpp
+++ b/Engine/Plugins/Runtime/GameFeatures/Source/GameFeatures/Private/GameFeaturePluginStateMachine.cpp
@@ -1660,8 +1660,7 @@ struct FBaseDataReleaseGameFeaturePluginState : public FGameFeaturePluginState
 		{
 			case (EGameFeaturePluginProtocol::InstallBundle):
 			{
-				//InstallBundleProtocols have a MetaData that controlls if they uninstall currently
-				return StateProperties.ProtocolOptions.GetSubtype<FInstallBundlePluginProtocolOptions>().bUninstallBeforeTerminate;
+				return true;
 			}
 
 			default:
@@ -1706,6 +1705,19 @@ struct FGameFeaturePluginState_Uninstalling : public FBaseDataReleaseGameFeature
 {
 	using FBaseDataReleaseGameFeaturePluginState::FBaseDataReleaseGameFeaturePluginState;
 
+	virtual void BeginState() override
+	{
+		if (ShouldVisitUninstallStateBeforeTerminal())
+		{
+			FBaseDataReleaseGameFeaturePluginState::BeginState();
+		}
+		else
+		{
+			Result = MakeValue();
+			bWasDeleted = true;
+		}
+	}
+
 	virtual EGameFeaturePluginState GetSuccessTransitionState() const override
 	{
 		return EGameFeaturePluginState::Uninstalled;
