Always call our callback when ReleaseAndUnregister even if we are not valid

#Synced-CL 52051466

diff --git a/Engine/Plugins/Runtime/GameFeatures/Source/GameFeatures/Private/GameFeatureStateHandle.cpp b/Engine/Plugins/Runtime/GameFeatures/Source/GameFeatures/Private/GameFeatureStateHandle.cpp
--- a/Engine/Plugins/Runtime/GameFeatures/Source/GameFeatures/Private/GameFeatureStateHandle.cpp
+++ b/Engine/Plugins/Runtime/GameFeatures/Source/GameFeatures/Private/GameFeatureStateHandle.cpp
@@ -58,6 +58,10 @@ void FGameFeatureStateHandle::ReleaseAndUnregister(TFunction<void(bool)> ResetCo
 		FGameFeatureStateHandleReferenceController::Get().UnregisterStateHandle(*this);
 		UniqueId.Invalidate();
 	}
+	else if (ResetCompleteCallback)
+	{
+		ResetCompleteCallback(false);
+	}
 }
 
 FGuid FGameFeatureStateHandle::GetUniqueId() const
