Update the Optional Content Installer with a new CVar:
GameFeatureOptionalContentInstaller.SkipUnkownBundles
This can be used so that the optional content installer will stop the request if it doesn't know about any of the bundles it was asked to install.
Add per plugin tracking for bundles. This way the optional installer can install the optional bundle that the current plugin is using. Then once the final plugin using that optional content is released the optional bundle is requested to be released.

diff --git a/Engine/Plugins/Experimental/DefaultInstallBundleManager/Source/Private/DefaultInstallBundleManager.cpp b/Engine/Plugins/Experimental/DefaultInstallBundleManager/Source/Private/DefaultInstallBundleManager.cpp
--- a/Engine/Plugins/Experimental/DefaultInstallBundleManager/Source/Private/DefaultInstallBundleManager.cpp
+++ b/Engine/Plugins/Experimental/DefaultInstallBundleManager/Source/Private/DefaultInstallBundleManager.cpp
@@ -3832,6 +3832,73 @@ TValueOrError<FInstallBundleCombinedInstallState, EInstallBundleResult> FDefault
 	return MakeValue(GetInstallStateInternal(RetBundleNames));
 }
 
+EInstallBundleResult FDefaultInstallBundleManager::GetAllMountedBundles(TArray<FName>& OutBundleNames) const
+{
+	CSV_SCOPED_TIMING_STAT(InstallBundleManager, InstallBundleManager_GetAllMountedBundles);
+
+	if (bUnrecoverableInitError || InitState == EInstallBundleManagerInitState::Failed)
+	{
+		return EInstallBundleResult::InitializationError;
+	}
+	if (InitState == EInstallBundleManagerInitState::NotInitialized)
+	{
+		return EInstallBundleResult::InitializationPending;
+	}
+
+	OutBundleNames.Reserve(BundleInfoMap.Num());
+	for (const TPair<FName, FBundleInfo>& BundleInfoPair : BundleInfoMap)
+	{
+		if (BundleInfoPair.Value.GetBundleStatus(*this) == EBundleState::Mounted)
+		{
+			OutBundleNames.Add(BundleInfoPair.Key);
+		}
+	}
+	return EInstallBundleResult::OK;
+}
+
+EInstallBundleResult FDefaultInstallBundleManager::GetAllBundlesWithState(EInstallBundleInstallState InstallState, TArray<FName>& OutBundleNames) const
+{
+	CSV_SCOPED_TIMING_STAT(InstallBundleManager, InstallBundleManager_GetAllBundlesWithState);
+
+	if (bUnrecoverableInitError || InitState == EInstallBundleManagerInitState::Failed)
+	{
+		return EInstallBundleResult::InitializationError;
+	}
+	if (InitState == EInstallBundleManagerInitState::NotInitialized)
+	{
+		return EInstallBundleResult::InitializationPending;
+	}
+
+	OutBundleNames.Reserve(BundleInfoMap.Num());
+	for (const TPair<FName, FBundleInfo>& BundleInfoPair : BundleInfoMap)
+	{
+		EInstallBundleInstallState CurrentBundleState = EInstallBundleInstallState::Count;
+		switch (BundleInfoPair.Value.GetBundleStatus(*this))
+		{
+		case EBundleState::NotInstalled:
+			CurrentBundleState = EInstallBundleInstallState::NotInstalled;
+			break;
+		case EBundleState::NeedsUpdate:
+			CurrentBundleState = EInstallBundleInstallState::NeedsUpdate;
+			break;
+		case EBundleState::NeedsMount:
+			CurrentBundleState = EInstallBundleInstallState::UpToDate;
+			break;
+		case EBundleState::Mounted:
+			CurrentBundleState = EInstallBundleInstallState::UpToDate;
+			break;
+		case EBundleState::Count:
+		default:
+			break;
+		}
+		if (CurrentBundleState == InstallState)
+		{
+			OutBundleNames.Add(BundleInfoPair.Key);
+		}
+	}
+	return EInstallBundleResult::OK;
+}
+
 void FDefaultInstallBundleManager::CancelAllGetInstallStateRequestsForTag(FName RequestTag)
 {
 	for (FGetInstallStateRequestRef& Request : GetInstallStateRequests)
diff --git a/Engine/Plugins/Experimental/DefaultInstallBundleManager/Source/Public/DefaultInstallBundleManager.h b/Engine/Plugins/Experimental/DefaultInstallBundleManager/Source/Public/DefaultInstallBundleManager.h
--- a/Engine/Plugins/Experimental/DefaultInstallBundleManager/Source/Public/DefaultInstallBundleManager.h
+++ b/Engine/Plugins/Experimental/DefaultInstallBundleManager/Source/Public/DefaultInstallBundleManager.h
@@ -665,6 +665,8 @@ public:
 	UE_API virtual FDelegateHandle GetInstallState(TArrayView<const FName> BundleNames, bool bAddDependencies, FInstallBundleGetInstallStateDelegate Callback, FName RequestTag = NAME_None) override;
 
 	UE_API virtual TValueOrError<FInstallBundleCombinedInstallState, EInstallBundleResult> GetInstallStateSynchronous(TArrayView<const FName> BundleNames, bool bAddDependencies) const override;
+	UE_API virtual EInstallBundleResult GetAllMountedBundles(TArray<FName>& OutBundleNames) const override;
+	UE_API virtual EInstallBundleResult GetAllBundlesWithState(EInstallBundleInstallState InstallState, TArray<FName>& OutBundleNames) const override;
 
 	UE_API virtual void CancelAllGetInstallStateRequestsForTag(FName RequestTag) override;
 	UE_API virtual void CancelAllGetInstallStateRequests(FDelegateHandle Handle) override;
diff --git a/Engine/Plugins/Runtime/GameFeatures/Source/GameFeatures/Private/GameFeatureOptionalContentInstaller.cpp b/Engine/Plugins/Runtime/GameFeatures/Source/GameFeatures/Private/GameFeatureOptionalContentInstaller.cpp
--- a/Engine/Plugins/Runtime/GameFeatures/Source/GameFeatures/Private/GameFeatureOptionalContentInstaller.cpp
+++ b/Engine/Plugins/Runtime/GameFeatures/Source/GameFeatures/Private/GameFeatureOptionalContentInstaller.cpp
@@ -19,6 +19,11 @@ namespace GameFeatureOptionalContentInstaller
 	static TAutoConsoleVariable<bool> CVarEnableOptionalContentInstaller(TEXT("GameFeatureOptionalContentInstaller.Enable"), 
 		true,
 		TEXT("Enable optional content installer"));
+
+	static bool bSkipRequestsForUnknownBundles = true;
+	static FAutoConsoleVariableRef CVarSkipRequestsForUnknownBundles(TEXT("GameFeatureOptionalContentInstaller.SkipUnkownBundles"),
+		bSkipRequestsForUnknownBundles,
+		TEXT("When true optional bundles that do not exist are skipped and the request continues."));
 }
 
 const FName UGameFeatureOptionalContentInstaller::GFOContentRequestName = TEXT("GFOContentRequest");
@@ -130,10 +135,19 @@ bool UGameFeatureOptionalContentInstaller::UpdateContent(const FString& PluginNa
 		if (MaybeInstallState.HasValue())
 		{
 			const FInstallBundleCombinedInstallState& InstallState = MaybeInstallState.GetValue();
-			bIsAvailable = Algo::AllOf(Bundles, [&InstallState](FName BundleName) { return InstallState.IndividualBundleStates.Contains(BundleName); });
-			bIsAvailable &= ensureMsgf(InstallState.IndividualBundleStates.Num() <= Bundles.Num(),
+			bIsAvailable = ensureMsgf(InstallState.IndividualBundleStates.Num() <= Bundles.Num(),
+
 				TEXT("UGameFeatureOptionalContentInstaller does not support dependencies tracking. Check optional install bundle dependencies for plugin %s."),
 				*PluginName);
+			if (GameFeatureOptionalContentInstaller::bSkipRequestsForUnknownBundles && bIsAvailable)
+			{
+				Bundles.Reset();
+				InstallState.IndividualBundleStates.GetKeys(Bundles);
+			}
+			else
+			{
+				bIsAvailable &= Algo::AllOf(Bundles, [&InstallState](FName BundleName) { return InstallState.IndividualBundleStates.Contains(BundleName); });
+			}
 		}
 	}
 
@@ -144,7 +158,7 @@ bool UGameFeatureOptionalContentInstaller::UpdateContent(const FString& PluginNa
 
 	for (const FName& Bundle : Bundles)
 	{
-		UE_LOG(LogGameFeatures, Log, TEXT("Requesting update for %s"), *Bundle.ToString());
+		UE_LOG(LogGameFeatures, Verbose, TEXT("Requesting update for %s"), *Bundle.ToString());
 	}
 	EInstallBundleRequestFlags InstallFlags = EInstallBundleRequestFlags::AsyncMount | EInstallBundleRequestFlags::ExplicitUpdateList;
 	if (bIsPredownload)
@@ -200,6 +214,12 @@ bool UGameFeatureOptionalContentInstaller::UpdateContent(const FString& PluginNa
 		Pending.BundlesEnqueued = MoveTemp(Request.BundlesEnqueued);
 		Pending.bIsPredownload = bIsPredownload;
 	}
+	const FName PluginFName(PluginName);
+	for (const FName& Bundle : Bundles)
+	{
+		TSet<FName>& PluginsUsingBundle = BundleToPlugins.FindOrAdd(Bundle);
+		PluginsUsingBundle.Add(PluginFName);
+	}
 
 	return true;
 }
@@ -220,7 +240,7 @@ void UGameFeatureOptionalContentInstaller::OnContentInstalled(FInstallBundleRequ
 
 	GFPInstall.BundlesEnqueued.Remove(InResult.BundleName);
 
-	UE_LOG(LogGameFeatures, Log, TEXT("Finished install for %s"), *InResult.BundleName.ToString());
+	UE_LOG(LogGameFeatures, Verbose, TEXT("Finished install for %s"), *InResult.BundleName.ToString());
 	if (InResult.Result != EInstallBundleResult::OK)
 	{
 		if (InResult.OptionalErrorCode.IsEmpty())
@@ -302,10 +322,31 @@ void UGameFeatureOptionalContentInstaller::ReleaseContent(const FString& PluginN
 	}
 	TArray<FName> KeepBundles = GetOptionalKeepBundlePredicate(PluginName);
 
+	TArray<FName> BundlesWithNoInUsePlugin;
+	const FName PluginFName(PluginName);
+	for (const FName& Bundle : Bundles)
+	{
+		TSet<FName>& PluginsUsingBundle = BundleToPlugins.FindOrAdd(Bundle);
+		PluginsUsingBundle.Remove(PluginFName);
+		if (PluginsUsingBundle.Num() == 0)
+		{
+			BundlesWithNoInUsePlugin.Add(Bundle);
+			BundleToPlugins.Remove(Bundle);
+		}
+	}
+	if (BundlesWithNoInUsePlugin.IsEmpty())
+	{
+		return;
+	}
+
+	for (const FName& Bundle : BundlesWithNoInUsePlugin)
+	{
+		UE_LOG(LogGameFeatures, Log, TEXT("Requesting Release for %s. No plugin is actively using it."), *Bundle.ToString());
+	}
 	Flags |= EInstallBundleReleaseRequestFlags::ExplicitRemoveList;
 
 	BundleManager->RequestReleaseContent(
-		Bundles, 
+		BundlesWithNoInUsePlugin,
 		Flags,
 		KeepBundles, 
 		GameFeatureOptionalContentInstaller::InstallBundleManagerVerbosityOverride);
@@ -316,6 +357,27 @@ void UGameFeatureOptionalContentInstaller::OnEnabled()
 	ensure(RelevantGFPs.IsEmpty());
 	RelevantGFPs.Empty();
 
+	// Capture a list of all existing bundles that are currently mounted outside of our system. We do not want to unmount something another system mounted.
+	{
+		BundleToPlugins.Reset();
+		const FName CoreInUseBundle(TEXT("#@#Core_Already_Mounted#@#"));
+
+		TArray<FName> AllBundles;
+		EInstallBundleResult InstallStateResult = BundleManager->GetAllMountedBundles(AllBundles);
+		if (InstallStateResult == EInstallBundleResult::OK)
+		{
+			for (const FName& MountedBundle : AllBundles)
+			{
+				TSet<FName>& PluginsUsingBundle = BundleToPlugins.FindOrAdd(MountedBundle);
+				PluginsUsingBundle.Add(CoreInUseBundle);
+			}
+		}
+		else
+		{
+			UE_LOG(LogGameFeatures, Error, TEXT("Failed to get all mounted bundles. Optional content installer enabled before budnle manager was initialized."));
+		}
+	}
+
 	UGameFeaturesSubsystem::Get().ForEachGameFeature([this](FGameFeatureInfo&& Info) -> void
 	{
 		if (Info.CurrentState >= EGameFeaturePluginState::Downloading)
@@ -337,6 +399,7 @@ void UGameFeatureOptionalContentInstaller::OnDisabled()
 
 	RelevantGFPs.Empty();
 	TotalProgressTracker.Reset();
+	BundleToPlugins.Reset();
 }
 
 bool UGameFeatureOptionalContentInstaller::IsEnabled() const
@@ -437,4 +500,4 @@ void UGameFeatureOptionalContentInstaller::OnGameFeatureReleasing(const FString&
 	ReleaseContent(PluginName);
 
 	RelevantGFPs.Remove(PluginName);
-}
+}
\ No newline at end of file
diff --git a/Engine/Plugins/Runtime/GameFeatures/Source/GameFeatures/Private/GameFeaturesSubsystem.cpp b/Engine/Plugins/Runtime/GameFeatures/Source/GameFeatures/Private/GameFeaturesSubsystem.cpp
--- a/Engine/Plugins/Runtime/GameFeatures/Source/GameFeatures/Private/GameFeaturesSubsystem.cpp
+++ b/Engine/Plugins/Runtime/GameFeatures/Source/GameFeatures/Private/GameFeaturesSubsystem.cpp
@@ -5392,6 +5392,31 @@ FString UGameFeaturesSubsystem::GetOptionalInstallBundleName(FStringView PluginN
 	return UGameFeatureData::GetOptionalInstallBundleName(PluginName, bEvenIfDoesntExist);
 }
 
+bool UGameFeaturesSubsystem::GetBundlesForPlugin(FStringView PluginName, TArray<FName>& OutBundles)
+{
+	FString PluginURL;
+	if (UGameFeaturesSubsystem::Get().GetPluginURLByName(PluginName, PluginURL))
+	{
+		return GetBundlesForPluginURL(PluginURL, OutBundles);
+	}
+	return false;
+}
+
+bool UGameFeaturesSubsystem::GetBundlesForPluginURL(FStringView PluginURL, TArray<FName>& OutBundles)
+{
+	if (const UGameFeaturePluginStateMachine* GFSM = FindGameFeaturePluginStateMachine(FString(PluginURL)))
+	{
+		const FGameFeatureProtocolMetadata& ProtocolMetadata = GetProtocolMetadata(GFSM);
+		if (ProtocolMetadata.HasSubtype<FInstallBundlePluginProtocolMetaData>())
+		{
+			const FInstallBundlePluginProtocolMetaData& ProtocolData = ProtocolMetadata.GetSubtype<FInstallBundlePluginProtocolMetaData>();
+			OutBundles.Append(ProtocolData.InstallBundles);
+			return true;
+		}
+	}
+	return false;
+}
+
 void UGameFeaturesSubsystem::SetExplanationForNotMountingPlugin(const FString& PluginURL, const FString& Explanation)
 {
 #if WITH_EDITOR
diff --git a/Engine/Plugins/Runtime/GameFeatures/Source/GameFeatures/Public/GameFeatureOptionalContentInstaller.h b/Engine/Plugins/Runtime/GameFeatures/Source/GameFeatures/Public/GameFeatureOptionalContentInstaller.h
--- a/Engine/Plugins/Runtime/GameFeatures/Source/GameFeatures/Public/GameFeatureOptionalContentInstaller.h
+++ b/Engine/Plugins/Runtime/GameFeatures/Source/GameFeatures/Public/GameFeatureOptionalContentInstaller.h
@@ -87,6 +87,8 @@ private:
 	TSharedPtr<IInstallBundleManager> BundleManager;
 	TSet<FString> RelevantGFPs;
 	TMap<FString, FGFPInstall> ActiveGFPInstalls;
+	/** Map of bundle to the plugins that requested them. Used for knowing when we can release an optional bundle. */
+	TMap<FName, TSet<FName>> BundleToPlugins;
 	TOptional<FInstallBundleCombinedProgressTracker> TotalProgressTracker;
 	static const FName GFOContentRequestName;
 
diff --git a/Engine/Plugins/Runtime/GameFeatures/Source/GameFeatures/Public/GameFeaturesSubsystem.h b/Engine/Plugins/Runtime/GameFeatures/Source/GameFeatures/Public/GameFeaturesSubsystem.h
--- a/Engine/Plugins/Runtime/GameFeatures/Source/GameFeatures/Public/GameFeaturesSubsystem.h
+++ b/Engine/Plugins/Runtime/GameFeatures/Source/GameFeatures/Public/GameFeaturesSubsystem.h
@@ -821,6 +821,16 @@ public:
 	 */
 	UE_API FString GetOptionalInstallBundleName(FStringView PluginName, bool bEvenIfDoesntExist = false);
 
+	/**
+	 * Returns the required bundles for this plugin based upon the current state machine.
+	 * Unlike GetInstallBundleName this will return all the bundles and dependencies for the given plugin.
+	 * @param - PluginName - the name of the GameFeaturePlugin we want to get the bundles for.
+	 * @param - OutBundles - Output reference array for the found bundles for the plugin.
+	 * @return true if an existing state machine with install bundle protocol is found.
+	 */
+	UE_API bool GetBundlesForPlugin(FStringView PluginName, TArray<FName>& OutBundles);
+	UE_API bool GetBundlesForPluginURL(FStringView PluginURL, TArray<FName>& OutBundles);
+
 private:
 	UE_API TSet<FString> GetActivePluginNames() const;
 
diff --git a/Engine/Source/Runtime/InstallBundleManager/Public/InstallBundleManagerInterface.h b/Engine/Source/Runtime/InstallBundleManager/Public/InstallBundleManagerInterface.h
--- a/Engine/Source/Runtime/InstallBundleManager/Public/InstallBundleManagerInterface.h
+++ b/Engine/Source/Runtime/InstallBundleManager/Public/InstallBundleManagerInterface.h
@@ -172,6 +172,8 @@ public:
 	virtual FDelegateHandle GetInstallState(TArrayView<const FName> BundleNames, bool bAddDependencies, FInstallBundleGetInstallStateDelegate Callback, FName RequestTag = NAME_None) = 0;
 	INSTALLBUNDLEMANAGER_API TValueOrError<FInstallBundleCombinedInstallState, EInstallBundleResult> GetInstallStateSynchronous(FName BundleName, bool bAddDependencies) const;
 	virtual TValueOrError<FInstallBundleCombinedInstallState, EInstallBundleResult> GetInstallStateSynchronous(TArrayView<const FName> BundleNames, bool bAddDependencies) const = 0;
+	virtual EInstallBundleResult GetAllMountedBundles(TArray<FName>& OutBundleNames) const = 0;
+	virtual EInstallBundleResult GetAllBundlesWithState(EInstallBundleInstallState InstallState, TArray<FName>& OutBundleNames) const = 0;
 	virtual void CancelAllGetInstallStateRequestsForTag(FName RequestTag) = 0;
 	virtual void CancelAllGetInstallStateRequests(FDelegateHandle Handle) = 0;
 
diff --git a/Engine/Source/Runtime/NullInstallBundleManager/Private/NullInstallBundleManager.cpp b/Engine/Source/Runtime/NullInstallBundleManager/Private/NullInstallBundleManager.cpp
--- a/Engine/Source/Runtime/NullInstallBundleManager/Private/NullInstallBundleManager.cpp
+++ b/Engine/Source/Runtime/NullInstallBundleManager/Private/NullInstallBundleManager.cpp
@@ -65,6 +65,16 @@ class FNullInstallBundleManager final : public IInstallBundleManager
 		return MakeValue(FInstallBundleCombinedInstallState());
 	}
 
+	virtual EInstallBundleResult GetAllMountedBundles(TArray<FName>& OutBundleNames) const override
+	{
+		return EInstallBundleResult::OK;
+	}
+
+	virtual EInstallBundleResult GetAllBundlesWithState(EInstallBundleInstallState InstallState, TArray<FName>& OutBundleNames) const override
+	{
+		return EInstallBundleResult::OK;
+	}
+
 	virtual void CancelAllGetInstallStateRequestsForTag(FName RequestTag) override
 	{
 	}
