Fix typo

#Synced-CL 52270238

diff --git a/Engine/Plugins/Runtime/GameFeatures/Source/GameFeatures/Private/GameFeaturesSubsystem.cpp b/Engine/Plugins/Runtime/GameFeatures/Source/GameFeatures/Private/GameFeaturesSubsystem.cpp
--- a/Engine/Plugins/Runtime/GameFeatures/Source/GameFeatures/Private/GameFeaturesSubsystem.cpp
+++ b/Engine/Plugins/Runtime/GameFeatures/Source/GameFeatures/Private/GameFeaturesSubsystem.cpp
@@ -2312,7 +2312,7 @@ namespace GameFeaturesSubsystem
 	};
 	ENUM_CLASS_FLAGS(ETransitionOptions)
 
-	struct FTransitionContenxt
+	struct FTransitionContext
 	{
 		TMap<FString, UE::GameFeatures::FResult> Results;
 		FMultipleGameFeaturePluginsReleased CompleteDelegate;
@@ -2325,7 +2325,7 @@ namespace GameFeaturesSubsystem
 		TOptional<FSuspendPrimaryAssetCleanup> SuspendPrimaryAssetsCleanup;
 		TOptional<FNotifyObservers> NotifyObservers;
 
-		UE_NODISCARD_CTOR FTransitionContenxt(ETransitionOptions TransitionOptions = ETransitionOptions::All)
+		UE_NODISCARD_CTOR FTransitionContext(ETransitionOptions TransitionOptions = ETransitionOptions::All)
 		{
 			if (EnumHasAnyFlags(TransitionOptions, ETransitionOptions::SuspendRenderAssetStreamer))
 			{
@@ -2364,7 +2364,7 @@ namespace GameFeaturesSubsystem
 			}
 		}
 
-		~FTransitionContenxt()
+		~FTransitionContext()
 		{
 			CompleteDelegate.ExecuteIfBound(Results);
 			// Explicit order of destruction is required
@@ -2403,7 +2403,7 @@ void UGameFeaturesSubsystem::DeactivateGameFeaturePlugin(TConstArrayView<FString
 {
 	// None for now, as the current list are more of an issue going from Registered -> <lower>
 	GameFeaturesSubsystem::ETransitionOptions TransitionOptions = GameFeaturesSubsystem::ETransitionOptions::None;
-	TSharedRef<GameFeaturesSubsystem::FTransitionContenxt> TransitionContext = MakeShared<GameFeaturesSubsystem::FTransitionContenxt>(TransitionOptions);
+	TSharedRef<GameFeaturesSubsystem::FTransitionContext> TransitionContext = MakeShared<GameFeaturesSubsystem::FTransitionContext>(TransitionOptions);
 	TransitionContext->CompleteDelegate = CompleteDelegate;
 
 	TransitionContext->Results.Reserve(PluginURLs.Num());
@@ -2462,7 +2462,7 @@ void UGameFeaturesSubsystem::UnloadGameFeaturePlugin(TConstArrayView<FString> Pl
 
 	EnumAddFlags(TransitionOptions, GameFeaturesSubsystem::ETransitionOptions::SuspendPrimaryAssetsCleanup);
 
-	TSharedRef<GameFeaturesSubsystem::FTransitionContenxt> TransitionContext = MakeShared<GameFeaturesSubsystem::FTransitionContenxt>(TransitionOptions);
+	TSharedRef<GameFeaturesSubsystem::FTransitionContext> TransitionContext = MakeShared<GameFeaturesSubsystem::FTransitionContext>(TransitionOptions);
 	TransitionContext->CompleteDelegate = CompleteDelegate;
 
 	TransitionContext->Results.Reserve(PluginURLs.Num());
@@ -2503,7 +2503,7 @@ void UGameFeaturesSubsystem::ReleaseGameFeaturePlugin(const FString& PluginURL,
 
 void UGameFeaturesSubsystem::ReleaseGameFeaturePlugin(TConstArrayView<FString> PluginURLs, const FMultipleGameFeaturePluginsReleased& CompleteDelegate)
 {
-	TSharedRef<GameFeaturesSubsystem::FTransitionContenxt> TransitionContext = MakeShared<GameFeaturesSubsystem::FTransitionContenxt>();
+	TSharedRef<GameFeaturesSubsystem::FTransitionContext> TransitionContext = MakeShared<GameFeaturesSubsystem::FTransitionContext>();
 	TransitionContext->CompleteDelegate = CompleteDelegate;
 
 	TransitionContext->Results.Reserve(PluginURLs.Num());
@@ -2623,7 +2623,7 @@ void UGameFeaturesSubsystem::TerminateGameFeaturePlugin(TConstArrayView<FString>
 		EnumRemoveFlags(TransitionOptions, GameFeaturesSubsystem::ETransitionOptions::SuspendRenderAssetStreamer);
 	}
 
-	TSharedRef<GameFeaturesSubsystem::FTransitionContenxt> TransitionContext = MakeShared<GameFeaturesSubsystem::FTransitionContenxt>(TransitionOptions);
+	TSharedRef<GameFeaturesSubsystem::FTransitionContext> TransitionContext = MakeShared<GameFeaturesSubsystem::FTransitionContext>(TransitionOptions);
 	TransitionContext->CompleteDelegate = CompleteDelegate;
 
 	TransitionContext->Results.Reserve(PluginURLs.Num());
