Optimize AssetManager cleanup during GameFeature unregistration in cases where only SpecificAssets are specified in PrimaryAssetTypesToScan.

#tests Windows Test client, before/after Insights captures
#rnx
#Synced-CL 52288889

diff --git a/Engine/Plugins/Runtime/GameFeatures/Source/GameFeatures/Private/GameFeaturesSubsystem.cpp b/Engine/Plugins/Runtime/GameFeatures/Source/GameFeatures/Private/GameFeaturesSubsystem.cpp
--- a/Engine/Plugins/Runtime/GameFeatures/Source/GameFeatures/Private/GameFeaturesSubsystem.cpp
+++ b/Engine/Plugins/Runtime/GameFeatures/Source/GameFeatures/Private/GameFeaturesSubsystem.cpp
@@ -320,7 +320,8 @@ namespace UE::GameFeatures
 
 	int32 SuspendPrimaryAssetCleanupCounter = 0;
 	TMultiMap<FString, FString> ScheduledAssetScanPathsByPluginName;
-	void CleanupPrimaryAssetByScanPaths(TConstArrayView<FName> ScanPaths)
+	TMultiMap<FString, FTopLevelAssetPath> ScheduledAssetClassesByPluginName;
+	void CleanupPrimaryAssetByScanPaths(TConstArrayView<FString> ScanPaths, TConstArrayView<FTopLevelAssetPath> AssetClasses)
 	{
 		if (!ScanPaths.IsEmpty())
 		{
@@ -328,17 +329,55 @@ namespace UE::GameFeatures
 			UAssetManager* AssetManager = UAssetManager::GetIfInitialized();
 			if (AssetRegistry && AssetManager)
 			{
-				TRACE_CPUPROFILER_EVENT_SCOPE(GFP_RemoveGameFeatureFromAssetManager_TraverseAR);
-				FARFilter AssetsFilter;
-				AssetsFilter.PackagePaths = ScanPaths;
-				AssetsFilter.bRecursivePaths = true;
-				AssetsFilter.TagsAndValues.Add(FPrimaryAssetId::PrimaryAssetTypeTag);
-				AssetsFilter.RecursiveClassPathsExclusionSet.Add(UGameFeatureData::StaticClass()->GetClassPathName());
-				AssetsFilter.bRecursiveClasses = true;
+				// Separate queries for specific assets and directories. In situations where only specific assets are provided this is a significant optimization.
+				// This matches logic from UAssetManager::SearchAssetRegistryPaths but as a combined query across multiple asset types
+				TArray<FName> PackageNames;
+				TArray<FName> PackagePaths;
+				for (const FString& Path : ScanPaths)
+				{
+					int32 DotIndex = INDEX_NONE;
+					if (Path.FindChar('.', DotIndex))
+					{
+						FString PackageName = Path.Mid(0, DotIndex); //avoid re-searching for index inside FPackageName::ObjectPathToPackageName
+
+						PackageNames.Add(FName(PackageName));
+					}
+					else if (Path.Len() > 0)
+					{
+						PackagePaths.Add(FName(UAssetManager::GetNormalizedPackagePath(Path, false)));
+					}
+				}
+
+				FARFilter AssetFilter;
+				AssetFilter.bRecursiveClasses = true;
+				AssetFilter.ClassPaths = AssetClasses;
 
 				TArray<FAssetData> AssetDatas;
-				if (AssetRegistry->GetInMemoryAssets(AssetsFilter, AssetDatas))
+				if (PackageNames.Num() > 0)
+				{
+					TRACE_CPUPROFILER_EVENT_SCOPE(GFP_RemoveGameFeatureFromAssetManager_QueryPackageNames);
+
+					AssetFilter.PackageNames = MoveTemp(PackageNames);
+
+					AssetRegistry->GetInMemoryAssets(AssetFilter, AssetDatas);
+				}
+
+				if (PackagePaths.Num() > 0)
 				{
+					TRACE_CPUPROFILER_EVENT_SCOPE(GFP_RemoveGameFeatureFromAssetManager_QueryPackagePaths);
+
+					//Clear PackageNames if it was set by the previous query
+					AssetFilter.PackageNames.Reset();
+
+					AssetFilter.bRecursivePaths = true;
+					AssetFilter.PackagePaths = MoveTemp(PackagePaths);
+
+					AssetRegistry->GetInMemoryAssets(AssetFilter, AssetDatas);
+				}
+
+				if (AssetDatas.Num() > 0)
+				{
+					TRACE_CPUPROFILER_EVENT_SCOPE(GFP_RemoveGameFeatureFromAssetManager_UnloadAssets);
 					TArray<FPrimaryAssetId> PotentialAssets;
 					PotentialAssets.Reserve(AssetDatas.Num());
 
@@ -954,6 +993,7 @@ void UGameFeaturesSubsystem::AddGameFeatureToAssetManager(const UGameFeatureData
 		UE::GameFeatures::SuspendPrimaryAssetCleanupCounter > 0)
 	{
 		UE::GameFeatures::ScheduledAssetScanPathsByPluginName.Remove(PluginName);
+		UE::GameFeatures::ScheduledAssetClassesByPluginName.Remove(PluginName);
 	}
 
 	const UAssetManagerSettings& Settings = LocalAssetManager.GetSettings();
@@ -974,6 +1014,7 @@ void UGameFeaturesSubsystem::RemoveGameFeatureFromAssetManager(const UGameFeatur
 
 	// Remove scan paths from asset manager
 	TArray<FString> ScanPaths;
+	TArray<FTopLevelAssetPath> AssetClasses;
 	for (FPrimaryAssetTypeInfo TypeInfo : GameFeatureToRemove->GetPrimaryAssetTypesToScan())
 	{
 		for (FDirectoryPath& Path : TypeInfo.GetDirectories())
@@ -991,6 +1032,7 @@ void UGameFeaturesSubsystem::RemoveGameFeatureFromAssetManager(const UGameFeatur
 		if (UE::GameFeatures::CVarUnloadPrimaryAssetsOnUnregister->GetBool())
 		{
 			ScanPaths.Append(TypeInfo.AssetScanPaths);
+			AssetClasses.Add(TypeInfo.AssetBaseClassLoaded->GetClassPathName());
 		}
 
 		if (AddedPrimaryAssetTypes.Contains(TypeInfo.PrimaryAssetType))
@@ -1011,13 +1053,16 @@ void UGameFeaturesSubsystem::RemoveGameFeatureFromAssetManager(const UGameFeatur
 			{
 				UE::GameFeatures::ScheduledAssetScanPathsByPluginName.Add(PluginName, Path);
 			}
+
+			UE::GameFeatures::ScheduledAssetClassesByPluginName.Reserve(UE::GameFeatures::ScheduledAssetClassesByPluginName.Num() + AssetClasses.Num());
+			for (const FTopLevelAssetPath& ClassPath : AssetClasses)
+			{
+				UE::GameFeatures::ScheduledAssetClassesByPluginName.Add(PluginName, ClassPath);
+			}
 		}
 		else
 		{
-			TArray<FName> ScanPathNames;
-			ScanPathNames.Reserve(ScanPaths.Num());
-			Algo::Transform(ScanPaths, ScanPathNames, [](const FString& Path) { return FName(Path); });
-			UE::GameFeatures::CleanupPrimaryAssetByScanPaths(ScanPathNames);
+			UE::GameFeatures::CleanupPrimaryAssetByScanPaths(ScanPaths, AssetClasses);
 		}
 	}
 }
@@ -2287,14 +2332,23 @@ namespace GameFeaturesSubsystem
 			ensure(UE::GameFeatures::SuspendPrimaryAssetCleanupCounter > 0);
 			if (--UE::GameFeatures::SuspendPrimaryAssetCleanupCounter == 0)
 			{
-				TArray<FName> ScanPaths;
+				TArray<FString> ScanPaths;
 				ScanPaths.Reserve(UE::GameFeatures::ScheduledAssetScanPathsByPluginName.Num());
 				for (const TTuple<FString, FString>& Pair : UE::GameFeatures::ScheduledAssetScanPathsByPluginName)
 				{
-					ScanPaths.Add(FName(Pair.Value));
+					ScanPaths.Add(Pair.Value);
 				}
 				UE::GameFeatures::ScheduledAssetScanPathsByPluginName.Empty();
-				UE::GameFeatures::CleanupPrimaryAssetByScanPaths(ScanPaths);
+
+				TArray<FTopLevelAssetPath> AssetClasses;
+				AssetClasses.Reserve(UE::GameFeatures::ScheduledAssetClassesByPluginName.Num());
+				for (const TTuple<FString, FTopLevelAssetPath>& Pair : UE::GameFeatures::ScheduledAssetClassesByPluginName)
+				{
+					AssetClasses.Add(Pair.Value);
+				}
+				UE::GameFeatures::ScheduledAssetClassesByPluginName.Empty();
+
+				UE::GameFeatures::CleanupPrimaryAssetByScanPaths(ScanPaths, AssetClasses);
 			}
 		}
 	};
