[GFP] Speedup AR searches during GFP transitions at the cost of temp memory usage.

#rnx
#Synced-CL 52288889

diff --git a/Engine/Plugins/Runtime/GameFeatures/Source/GameFeatures/Private/GameFeaturesSubsystem.cpp b/Engine/Plugins/Runtime/GameFeatures/Source/GameFeatures/Private/GameFeaturesSubsystem.cpp
--- a/Engine/Plugins/Runtime/GameFeatures/Source/GameFeatures/Private/GameFeaturesSubsystem.cpp
+++ b/Engine/Plugins/Runtime/GameFeatures/Source/GameFeatures/Private/GameFeaturesSubsystem.cpp
@@ -8,6 +8,7 @@
 #include "Algo/Transform.h"
 #include "Algo/Unique.h"
 #include "AssetRegistry/AssetData.h"
+#include "AssetRegistry/AssetRegistryHelpers.h"
 #include "AssetRegistry/IAssetRegistry.h"
 #include "Async/Mutex.h"
 #include "Async/ParallelFor.h"
@@ -2362,7 +2363,8 @@ namespace GameFeaturesSubsystem
 		NotifyObservers					= 1 << 3,
 		SuspendRenderAssetStreamer		= 1 << 4,
 		SuspendPrimaryAssetsCleanup		= 1 << 5,
-		All = SuspendPluginManager | SuspendGameplayTagTreeBroadcast | SuspendAssetRegistry | NotifyObservers | SuspendRenderAssetStreamer | SuspendPrimaryAssetsCleanup,
+		CacheAssetRegistryQueries		= 1 << 6,
+		All = SuspendPluginManager | SuspendGameplayTagTreeBroadcast | SuspendAssetRegistry | NotifyObservers | SuspendRenderAssetStreamer | SuspendPrimaryAssetsCleanup | CacheAssetRegistryQueries,
 	};
 	ENUM_CLASS_FLAGS(ETransitionOptions)
 
@@ -2376,6 +2378,7 @@ namespace GameFeaturesSubsystem
 		TOptional<FSuspendPluginManager> SuspendPluginManager;
 		TOptional<FSuspendGameplayTagTreeChangedBroadcast> SuspendGameplayTagTreeBroadcast;
 		TOptional<IAssetRegistry::FSuspendPathsDismounting> SuspendAssetRegistry;
+		TOptional<UAssetRegistryHelpers::FTemporaryCachingModeScope> CacheAssetRegistryQueries;
 		TOptional<FSuspendPrimaryAssetCleanup> SuspendPrimaryAssetsCleanup;
 		TOptional<FNotifyObservers> NotifyObservers;
 
@@ -2407,6 +2410,11 @@ namespace GameFeaturesSubsystem
 				SuspendAssetRegistry.Emplace();
 			}
 
+			if (EnumHasAnyFlags(TransitionOptions, ETransitionOptions::CacheAssetRegistryQueries))
+			{
+				CacheAssetRegistryQueries.Emplace(true);
+			}
+
 			if (EnumHasAnyFlags(TransitionOptions, ETransitionOptions::SuspendPrimaryAssetsCleanup))
 			{
 				SuspendPrimaryAssetsCleanup.Emplace();
@@ -2428,6 +2436,7 @@ namespace GameFeaturesSubsystem
 			// i.e. other batches need to finish.
 			NotifyObservers.Reset(); // Some observers are responsible for clearing references; must happen before SuspendPluginManager.Reset()
 			SuspendPrimaryAssetsCleanup.Reset(); // assets cleanup requires data in asset registry still present
+			CacheAssetRegistryQueries.Reset();
 			SuspendAssetRegistry.Reset();
 			SuspendGameplayTagTreeBroadcast.Reset();
 			SuspendPluginManager.Reset();
