[Backout] - CL52315035
#rnx
Original CL Desc
-----------------------------------------------------------------
Optimize AssetManager cleanup during GameFeature unregistration in cases where only SpecificAssets are specified in PrimaryAssetTypesToScan.

#tests Windows Test client, before/after Insights captures
#rnx

diff --git a/Engine/Plugins/Runtime/GameFeatures/Source/GameFeatures/Private/GameFeaturesSubsystem.cpp b/Engine/Plugins/Runtime/GameFeatures/Source/GameFeatures/Private/GameFeaturesSubsystem.cpp
--- a/Engine/Plugins/Runtime/GameFeatures/Source/GameFeatures/Private/GameFeaturesSubsystem.cpp
+++ b/Engine/Plugins/Runtime/GameFeatures/Source/GameFeatures/Private/GameFeaturesSubsystem.cpp
@@ -321,8 +321,7 @@ namespace UE::GameFeatures
 
 	int32 SuspendPrimaryAssetCleanupCounter = 0;
 	TMultiMap<FString, FString> ScheduledAssetScanPathsByPluginName;
-	TMultiMap<FString, FTopLevelAssetPath> ScheduledAssetClassesByPluginName;
-	void CleanupPrimaryAssetByScanPaths(TConstArrayView<FString> ScanPaths, TConstArrayView<FTopLevelAssetPath> AssetClasses)
+	void CleanupPrimaryAssetByScanPaths(TConstArrayView<FName> ScanPaths)
 	{
 		if (!ScanPaths.IsEmpty())
 		{
@@ -330,55 +329,17 @@ namespace UE::GameFeatures
 			UAssetManager* AssetManager = UAssetManager::GetIfInitialized();
 			if (AssetRegistry && AssetManager)
 			{
-				// Separate queries for specific assets and directories. In situations where only specific assets are provided this is a significant optimization.
-				// This matches logic from UAssetManager::SearchAssetRegistryPaths but as a combined query across multiple asset types
-				TArray<FName> PackageNames;
-				TArray<FName> PackagePaths;
-				for (const FString& Path : ScanPaths)
-				{
-					int32 DotIndex = INDEX_NONE;
-					if (Path.FindChar('.', DotIndex))
-					{
-						FString PackageName = Path.Mid(0, DotIndex); //avoid re-searching for index inside FPackageName::ObjectPathToPackageName
-
-						PackageNames.Add(FName(PackageName));
-					}
-					else if (Path.Len() > 0)
-					{
-						PackagePaths.Add(FName(UAssetManager::GetNormalizedPackagePath(Path, false)));
-					}
-				}
-
-				FARFilter AssetFilter;
-				AssetFilter.bRecursiveClasses = true;
-				AssetFilter.ClassPaths = AssetClasses;
+				TRACE_CPUPROFILER_EVENT_SCOPE(GFP_RemoveGameFeatureFromAssetManager_TraverseAR);
+				FARFilter AssetsFilter;
+				AssetsFilter.PackagePaths = ScanPaths;
+				AssetsFilter.bRecursivePaths = true;
+				AssetsFilter.TagsAndValues.Add(FPrimaryAssetId::PrimaryAssetTypeTag);
+				AssetsFilter.RecursiveClassPathsExclusionSet.Add(UGameFeatureData::StaticClass()->GetClassPathName());
+				AssetsFilter.bRecursiveClasses = true;
 
 				TArray<FAssetData> AssetDatas;
-				if (PackageNames.Num() > 0)
-				{
-					TRACE_CPUPROFILER_EVENT_SCOPE(GFP_RemoveGameFeatureFromAssetManager_QueryPackageNames);
-
-					AssetFilter.PackageNames = MoveTemp(PackageNames);
-
-					AssetRegistry->GetInMemoryAssets(AssetFilter, AssetDatas);
-				}
-
-				if (PackagePaths.Num() > 0)
+				if (AssetRegistry->GetInMemoryAssets(AssetsFilter, AssetDatas))
 				{
-					TRACE_CPUPROFILER_EVENT_SCOPE(GFP_RemoveGameFeatureFromAssetManager_QueryPackagePaths);
-
-					//Clear PackageNames if it was set by the previous query
-					AssetFilter.PackageNames.Reset();
-
-					AssetFilter.bRecursivePaths = true;
-					AssetFilter.PackagePaths = MoveTemp(PackagePaths);
-
-					AssetRegistry->GetInMemoryAssets(AssetFilter, AssetDatas);
-				}
-
-				if (AssetDatas.Num() > 0)
-				{
-					TRACE_CPUPROFILER_EVENT_SCOPE(GFP_RemoveGameFeatureFromAssetManager_UnloadAssets);
 					TArray<FPrimaryAssetId> PotentialAssets;
 					PotentialAssets.Reserve(AssetDatas.Num());
 
@@ -994,7 +955,6 @@ void UGameFeaturesSubsystem::AddGameFeatureToAssetManager(const UGameFeatureData
 		UE::GameFeatures::SuspendPrimaryAssetCleanupCounter > 0)
 	{
 		UE::GameFeatures::ScheduledAssetScanPathsByPluginName.Remove(PluginName);
-		UE::GameFeatures::ScheduledAssetClassesByPluginName.Remove(PluginName);
 	}
 
 	const UAssetManagerSettings& Settings = LocalAssetManager.GetSettings();
@@ -1015,7 +975,6 @@ void UGameFeaturesSubsystem::RemoveGameFeatureFromAssetManager(const UGameFeatur
 
 	// Remove scan paths from asset manager
 	TArray<FString> ScanPaths;
-	TArray<FTopLevelAssetPath> AssetClasses;
 	for (FPrimaryAssetTypeInfo TypeInfo : GameFeatureToRemove->GetPrimaryAssetTypesToScan())
 	{
 		for (FDirectoryPath& Path : TypeInfo.GetDirectories())
@@ -1033,7 +992,6 @@ void UGameFeaturesSubsystem::RemoveGameFeatureFromAssetManager(const UGameFeatur
 		if (UE::GameFeatures::CVarUnloadPrimaryAssetsOnUnregister->GetBool())
 		{
 			ScanPaths.Append(TypeInfo.AssetScanPaths);
-			AssetClasses.Add(TypeInfo.AssetBaseClassLoaded->GetClassPathName());
 		}
 
 		if (AddedPrimaryAssetTypes.Contains(TypeInfo.PrimaryAssetType))
@@ -1054,16 +1012,13 @@ void UGameFeaturesSubsystem::RemoveGameFeatureFromAssetManager(const UGameFeatur
 			{
 				UE::GameFeatures::ScheduledAssetScanPathsByPluginName.Add(PluginName, Path);
 			}
-
-			UE::GameFeatures::ScheduledAssetClassesByPluginName.Reserve(UE::GameFeatures::ScheduledAssetClassesByPluginName.Num() + AssetClasses.Num());
-			for (const FTopLevelAssetPath& ClassPath : AssetClasses)
-			{
-				UE::GameFeatures::ScheduledAssetClassesByPluginName.Add(PluginName, ClassPath);
-			}
 		}
 		else
 		{
-			UE::GameFeatures::CleanupPrimaryAssetByScanPaths(ScanPaths, AssetClasses);
+			TArray<FName> ScanPathNames;
+			ScanPathNames.Reserve(ScanPaths.Num());
+			Algo::Transform(ScanPaths, ScanPathNames, [](const FString& Path) { return FName(Path); });
+			UE::GameFeatures::CleanupPrimaryAssetByScanPaths(ScanPathNames);
 		}
 	}
 }
@@ -2333,23 +2288,14 @@ namespace GameFeaturesSubsystem
 			ensure(UE::GameFeatures::SuspendPrimaryAssetCleanupCounter > 0);
 			if (--UE::GameFeatures::SuspendPrimaryAssetCleanupCounter == 0)
 			{
-				TArray<FString> ScanPaths;
+				TArray<FName> ScanPaths;
 				ScanPaths.Reserve(UE::GameFeatures::ScheduledAssetScanPathsByPluginName.Num());
 				for (const TTuple<FString, FString>& Pair : UE::GameFeatures::ScheduledAssetScanPathsByPluginName)
 				{
-					ScanPaths.Add(Pair.Value);
+					ScanPaths.Add(FName(Pair.Value));
 				}
 				UE::GameFeatures::ScheduledAssetScanPathsByPluginName.Empty();
-
-				TArray<FTopLevelAssetPath> AssetClasses;
-				AssetClasses.Reserve(UE::GameFeatures::ScheduledAssetClassesByPluginName.Num());
-				for (const TTuple<FString, FTopLevelAssetPath>& Pair : UE::GameFeatures::ScheduledAssetClassesByPluginName)
-				{
-					AssetClasses.Add(Pair.Value);
-				}
-				UE::GameFeatures::ScheduledAssetClassesByPluginName.Empty();
-
-				UE::GameFeatures::CleanupPrimaryAssetByScanPaths(ScanPaths, AssetClasses);
+				UE::GameFeatures::CleanupPrimaryAssetByScanPaths(ScanPaths);
 			}
 		}
 	};
