[Backout] - CL52316578
#rnx
Original CL Desc
-----------------------------------------------------------------
[GFP] Speedup AR searches during GFP transitions at the cost of temp memory usage.

#rnx
#Synced-CL 52288889

diff --git a/Engine/Plugins/Runtime/GameFeatures/Source/GameFeatures/Private/GameFeaturesSubsystem.cpp b/Engine/Plugins/Runtime/GameFeatures/Source/GameFeatures/Private/GameFeaturesSubsystem.cpp
--- a/Engine/Plugins/Runtime/GameFeatures/Source/GameFeatures/Private/GameFeaturesSubsystem.cpp
+++ b/Engine/Plugins/Runtime/GameFeatures/Source/GameFeatures/Private/GameFeaturesSubsystem.cpp
@@ -8,7 +8,6 @@
 #include "Algo/Transform.h"
 #include "Algo/Unique.h"
 #include "AssetRegistry/AssetData.h"
-#include "AssetRegistry/AssetRegistryHelpers.h"
 #include "AssetRegistry/IAssetRegistry.h"
 #include "Async/Mutex.h"
 #include "Async/ParallelFor.h"
@@ -2309,12 +2308,11 @@ namespace GameFeaturesSubsystem
 		NotifyObservers					= 1 << 3,
 		SuspendRenderAssetStreamer		= 1 << 4,
 		SuspendPrimaryAssetsCleanup		= 1 << 5,
-		CacheAssetRegistryQueries		= 1 << 6,
-		All = SuspendPluginManager | SuspendGameplayTagTreeBroadcast | SuspendAssetRegistry | NotifyObservers | SuspendRenderAssetStreamer | SuspendPrimaryAssetsCleanup | CacheAssetRegistryQueries,
+		All = SuspendPluginManager | SuspendGameplayTagTreeBroadcast | SuspendAssetRegistry | NotifyObservers | SuspendRenderAssetStreamer | SuspendPrimaryAssetsCleanup,
 	};
 	ENUM_CLASS_FLAGS(ETransitionOptions)
 
-	struct FTransitionContext
+	struct FTransitionContenxt
 	{
 		TMap<FString, UE::GameFeatures::FResult> Results;
 		FMultipleGameFeaturePluginsReleased CompleteDelegate;
@@ -2324,11 +2322,10 @@ namespace GameFeaturesSubsystem
 		TOptional<FSuspendPluginManager> SuspendPluginManager;
 		TOptional<FSuspendGameplayTagTreeChangedBroadcast> SuspendGameplayTagTreeBroadcast;
 		TOptional<IAssetRegistry::FSuspendPathsDismounting> SuspendAssetRegistry;
-		TOptional<UAssetRegistryHelpers::FTemporaryCachingModeScope> CacheAssetRegistryQueries;
 		TOptional<FSuspendPrimaryAssetCleanup> SuspendPrimaryAssetsCleanup;
 		TOptional<FNotifyObservers> NotifyObservers;
 
-		UE_NODISCARD_CTOR FTransitionContext(ETransitionOptions TransitionOptions = ETransitionOptions::All)
+		UE_NODISCARD_CTOR FTransitionContenxt(ETransitionOptions TransitionOptions = ETransitionOptions::All)
 		{
 			if (EnumHasAnyFlags(TransitionOptions, ETransitionOptions::SuspendRenderAssetStreamer))
 			{
@@ -2356,11 +2353,6 @@ namespace GameFeaturesSubsystem
 				SuspendAssetRegistry.Emplace();
 			}
 
-			if (EnumHasAnyFlags(TransitionOptions, ETransitionOptions::CacheAssetRegistryQueries))
-			{
-				CacheAssetRegistryQueries.Emplace(true);
-			}
-
 			if (EnumHasAnyFlags(TransitionOptions, ETransitionOptions::SuspendPrimaryAssetsCleanup))
 			{
 				SuspendPrimaryAssetsCleanup.Emplace();
@@ -2372,7 +2364,7 @@ namespace GameFeaturesSubsystem
 			}
 		}
 
-		~FTransitionContext()
+		~FTransitionContenxt()
 		{
 			CompleteDelegate.ExecuteIfBound(Results);
 			// Explicit order of destruction is required
@@ -2382,7 +2374,6 @@ namespace GameFeaturesSubsystem
 			// i.e. other batches need to finish.
 			NotifyObservers.Reset(); // Some observers are responsible for clearing references; must happen before SuspendPluginManager.Reset()
 			SuspendPrimaryAssetsCleanup.Reset(); // assets cleanup requires data in asset registry still present
-			CacheAssetRegistryQueries.Reset();
 			SuspendAssetRegistry.Reset();
 			SuspendGameplayTagTreeBroadcast.Reset();
 			SuspendPluginManager.Reset();
@@ -2412,7 +2403,7 @@ void UGameFeaturesSubsystem::DeactivateGameFeaturePlugin(TConstArrayView<FString
 {
 	// None for now, as the current list are more of an issue going from Registered -> <lower>
 	GameFeaturesSubsystem::ETransitionOptions TransitionOptions = GameFeaturesSubsystem::ETransitionOptions::None;
-	TSharedRef<GameFeaturesSubsystem::FTransitionContext> TransitionContext = MakeShared<GameFeaturesSubsystem::FTransitionContext>(TransitionOptions);
+	TSharedRef<GameFeaturesSubsystem::FTransitionContenxt> TransitionContext = MakeShared<GameFeaturesSubsystem::FTransitionContenxt>(TransitionOptions);
 	TransitionContext->CompleteDelegate = CompleteDelegate;
 
 	TransitionContext->Results.Reserve(PluginURLs.Num());
@@ -2471,7 +2462,7 @@ void UGameFeaturesSubsystem::UnloadGameFeaturePlugin(TConstArrayView<FString> Pl
 
 	EnumAddFlags(TransitionOptions, GameFeaturesSubsystem::ETransitionOptions::SuspendPrimaryAssetsCleanup);
 
-	TSharedRef<GameFeaturesSubsystem::FTransitionContext> TransitionContext = MakeShared<GameFeaturesSubsystem::FTransitionContext>(TransitionOptions);
+	TSharedRef<GameFeaturesSubsystem::FTransitionContenxt> TransitionContext = MakeShared<GameFeaturesSubsystem::FTransitionContenxt>(TransitionOptions);
 	TransitionContext->CompleteDelegate = CompleteDelegate;
 
 	TransitionContext->Results.Reserve(PluginURLs.Num());
@@ -2512,7 +2503,7 @@ void UGameFeaturesSubsystem::ReleaseGameFeaturePlugin(const FString& PluginURL,
 
 void UGameFeaturesSubsystem::ReleaseGameFeaturePlugin(TConstArrayView<FString> PluginURLs, const FMultipleGameFeaturePluginsReleased& CompleteDelegate)
 {
-	TSharedRef<GameFeaturesSubsystem::FTransitionContext> TransitionContext = MakeShared<GameFeaturesSubsystem::FTransitionContext>();
+	TSharedRef<GameFeaturesSubsystem::FTransitionContenxt> TransitionContext = MakeShared<GameFeaturesSubsystem::FTransitionContenxt>();
 	TransitionContext->CompleteDelegate = CompleteDelegate;
 
 	TransitionContext->Results.Reserve(PluginURLs.Num());
@@ -2632,7 +2623,7 @@ void UGameFeaturesSubsystem::TerminateGameFeaturePlugin(TConstArrayView<FString>
 		EnumRemoveFlags(TransitionOptions, GameFeaturesSubsystem::ETransitionOptions::SuspendRenderAssetStreamer);
 	}
 
-	TSharedRef<GameFeaturesSubsystem::FTransitionContext> TransitionContext = MakeShared<GameFeaturesSubsystem::FTransitionContext>(TransitionOptions);
+	TSharedRef<GameFeaturesSubsystem::FTransitionContenxt> TransitionContext = MakeShared<GameFeaturesSubsystem::FTransitionContenxt>(TransitionOptions);
 	TransitionContext->CompleteDelegate = CompleteDelegate;
 
 	TransitionContext->Results.Reserve(PluginURLs.Num());
