- Allow GFPs in AdditionalPluginDirectories

[RN] GameFeaturePlugins can now work from locations listed in AdditionalPluginDirectories (in a .uproject)
#Synced-CL 52482450

diff --git a/Engine/Plugins/Runtime/GameFeatures/Source/GameFeatures/Private/GameFeaturesSubsystemSettings.cpp b/Engine/Plugins/Runtime/GameFeatures/Source/GameFeatures/Private/GameFeaturesSubsystemSettings.cpp
--- a/Engine/Plugins/Runtime/GameFeatures/Source/GameFeatures/Private/GameFeaturesSubsystemSettings.cpp
+++ b/Engine/Plugins/Runtime/GameFeatures/Source/GameFeatures/Private/GameFeaturesSubsystemSettings.cpp
@@ -1,8 +1,10 @@
 // Copyright Epic Games, Inc. All Rights Reserved.
 
 #include "GameFeaturesSubsystemSettings.h"
+#include "Interfaces/IProjectManager.h"
 #include "Misc/Paths.h"
 #include "Misc/PathViews.h"
+#include "ProjectDescriptor.h"
 
 #include UE_INLINE_GENERATED_CPP_BY_NAME(GameFeaturesSubsystemSettings)
 
@@ -30,22 +32,37 @@ bool UGameFeaturesSubsystemSettings::IsValidGameFeaturePlugin(FStringView Plugin
 				FPaths::EGetExtensionDirsFlags::WithRestricted;
 
 			// Get all the existing game feature paths
-			TArray<FString> RelativePaths = FPaths::GetExtensionDirs(
+			TArray<FString> PluginPaths = FPaths::GetExtensionDirs(
 				FPaths::ProjectDir(), FPaths::Combine(TEXT("Plugins"), TEXT("GameFeatures")), ExtensionFlags);
 
 			// The base directory may not exist yet, add it if empty
-			if (RelativePaths.IsEmpty())
+			if (PluginPaths.IsEmpty())
 			{
-				RelativePaths.Add(FPaths::Combine(FPaths::ProjectDir(), TEXT("Plugins"), TEXT("GameFeatures")));
+				PluginPaths.Add(FPaths::Combine(FPaths::ProjectDir(), TEXT("Plugins"), TEXT("GameFeatures")));
 			}
 
-			BuiltInGameFeaturePluginsFolders.Reserve(2 * RelativePaths.Num());
-			for (FString& BuiltInFolder : RelativePaths)
+			// Add any additional plugin directories from the project
+			const FProjectDescriptor* Project = IProjectManager::Get().GetCurrentProject();
+			if (Project != nullptr)
+			{
+				for (const FString& Dir : Project->GetAdditionalPluginDirectories())
+				{
+					PluginPaths.Add(FPaths::Combine(Dir, TEXT("GameFeatures")));
+				}
+			}
+
+			BuiltInGameFeaturePluginsFolders.Reserve(2 * PluginPaths.Num());
+			for (FString& BuiltInFolder : PluginPaths)
 			{
 				BuiltInFolder /= TEXT(""); // Add trailing slash if needed
-				BuiltInGameFeaturePluginsFolders.Add(FPaths::ConvertRelativePathToFull(BuiltInFolder));
+				// if the path was absolute, no need to do this
+				if (FPaths::IsRelative(BuiltInFolder))
+				{
+					BuiltInGameFeaturePluginsFolders.Add(FPaths::ConvertRelativePathToFull(BuiltInFolder));
+				}
 				BuiltInGameFeaturePluginsFolders.Add(MoveTemp(BuiltInFolder));
 			}
+			BuiltInGameFeaturePluginsFolders.Shrink();
 		}
 
 		TArray<FString> BuiltInGameFeaturePluginsFolders;
