Add CVar to skip built-in GFP loading (GameFeaturePlugin.SkipBuiltInGameFeatures)

#Synced-CL 52158602
#Synced-CL 52158602

diff --git a/Engine/Plugins/Runtime/GameFeatures/Source/GameFeatures/Private/GameFeaturesSubsystem.cpp b/Engine/Plugins/Runtime/GameFeatures/Source/GameFeatures/Private/GameFeaturesSubsystem.cpp
--- a/Engine/Plugins/Runtime/GameFeatures/Source/GameFeatures/Private/GameFeaturesSubsystem.cpp
+++ b/Engine/Plugins/Runtime/GameFeatures/Source/GameFeatures/Private/GameFeaturesSubsystem.cpp
@@ -116,6 +116,13 @@ namespace UE::GameFeatures
 		ECVF_Default);
 #endif // !UE_BUILD_SHIPPING
 
+	static bool bSkipBuiltInGameFeatures = false;
+	static FAutoConsoleVariableRef CVarSkipBuiltInGameFeatures(
+		TEXT("GameFeaturePlugin.SkipBuiltInGameFeatures"),
+		bSkipBuiltInGameFeatures,
+		TEXT("Skip loading built-in game feature plugins."),
+		ECVF_Default);
+
 	static bool bTrimNonStartupEnabledPlugins = true;
 	static FAutoConsoleVariableRef CVarTrimNonBuiltInPlugins(
 		TEXT("GameFeaturePlugin.TrimNonStartupEnabledPlugins"),
@@ -3056,6 +3063,13 @@ void UGameFeaturesSubsystem::LoadBuiltInGameFeaturePluginsInternal(TArray<TShare
 {
 	TRACE_CPUPROFILER_EVENT_SCOPE_STR("UGameFeaturesSubsystem::LoadBuiltInGameFeaturePlugins");
 
+	if (UE::GameFeatures::bSkipBuiltInGameFeatures)
+	{
+		UE_LOG(LogGameFeatures, Log, TEXT("Skipping built-in GFP loading (GameFeaturePlugin.SkipBuiltInGameFeatures)"));
+		InCompleteDelegate.ExecuteIfBound(TMap<FString, UE::GameFeatures::FResult>());
+		return;
+	}
+
 	TSharedRef<FAmmortizedLoadContext> LoadContext = MakeShared<FAmmortizedLoadContext>();
 	FGuid ReferenceGuid = FGuid::NewGuid();
 	LoadContext->GUID = ReferenceGuid;
