Fixed a bunch of unreachable code warnings causing errors on some targets

#Synced-CL 52609397

#Synced-CL 52609397

diff --git a/Engine/Plugins/Experimental/RemoteSession/Source/RemoteSession/Private/RemoteSessionRole.cpp b/Engine/Plugins/Experimental/RemoteSession/Source/RemoteSession/Private/RemoteSessionRole.cpp
--- a/Engine/Plugins/Experimental/RemoteSession/Source/RemoteSession/Private/RemoteSessionRole.cpp
+++ b/Engine/Plugins/Experimental/RemoteSession/Source/RemoteSession/Private/RemoteSessionRole.cpp
@@ -392,7 +392,8 @@ void FRemoteSessionRole::OnReceiveLegacyVersion(IBackChannelPacket& Message)
 	else
 	{
 		// this path should not be possible..
-		if (ensureMsgf(false, TEXT("Received new protocol version through legacy handshake!")))
+		bool value = ensureMsgf(false, TEXT("Received new protocol version through legacy handshake!"));
+		if (value)
 		{
 			RemoteVersion = VersionString;
 		}
diff --git a/Engine/Plugins/Messaging/UdpMessaging/Source/UdpMessaging/Private/Transport/UdpMessageTransport.cpp b/Engine/Plugins/Messaging/UdpMessaging/Source/UdpMessaging/Private/Transport/UdpMessageTransport.cpp
--- a/Engine/Plugins/Messaging/UdpMessaging/Source/UdpMessaging/Private/Transport/UdpMessageTransport.cpp
+++ b/Engine/Plugins/Messaging/UdpMessaging/Source/UdpMessaging/Private/Transport/UdpMessageTransport.cpp
@@ -447,11 +447,13 @@ bool FUdpMessageTransport::TransportMessage(const TSharedRef<IMessageContext, ES
 		return false;
 	}
 
+#if !NO_LOGGING
 	if (UE_GET_LOG_VERBOSITY(LogUdpMessaging) >= ELogVerbosity::Verbose)
 	{
 		FString RecipientStr = FString::JoinBy(Recipients, TEXT("+"), [](const FGuid& Guid) { return Guid.ToString(); });
 		UE_LOG(LogUdpMessaging, Verbose, TEXT("TransportMessage %s from %s to %s"), *Context->GetMessageTypePathName().ToString(), *Context->GetSender().ToString(), *RecipientStr);
 	}
+#endif
 
 	return MessageProcessor->EnqueueOutboundMessage(Context, Recipients);
 }
diff --git a/Engine/Plugins/Runtime/GameFeatures/Source/GameFeatures/Private/GameFeatureData.cpp b/Engine/Plugins/Runtime/GameFeatures/Source/GameFeatures/Private/GameFeatureData.cpp
--- a/Engine/Plugins/Runtime/GameFeatures/Source/GameFeatures/Private/GameFeatureData.cpp
+++ b/Engine/Plugins/Runtime/GameFeatures/Source/GameFeatures/Private/GameFeatureData.cpp
@@ -700,6 +700,7 @@ void UGameFeatureData::InitializeHierarchicalPluginIniFiles(const TArrayView<FSt
 			PluginNames.Emplace(*FPaths::GetBaseFilename(PluginInstalledFilename));
 		}
 
+#if !NO_LOGGING
 		if (UE_GET_LOG_VERBOSITY(LogGameFeatures) >= ELogVerbosity::Verbose)
 		{
 			FString PluginList = FString::Printf(TEXT("[%s]"), *Algo::Accumulate(PluginInstalledFilenames, FString(), [](FString InResult, const FString& InName)
@@ -709,6 +710,7 @@ void UGameFeatureData::InitializeHierarchicalPluginIniFiles(const TArrayView<FSt
 			}));
 			UE_LOG(LogGameFeatures, Verbose, TEXT("Loading GameFeature config modification for %s"), *PluginList);
 		}
+#endif
 
 		UE::DynamicConfig::PerformDynamicConfig("InitializeHierarchicalPluginIniFiles", [&PluginNames](FConfigModificationTracker* ChangeTracker)
 			{
diff --git a/Engine/Source/Runtime/Messaging/Private/Bridge/MessageBridge.cpp b/Engine/Source/Runtime/Messaging/Private/Bridge/MessageBridge.cpp
--- a/Engine/Source/Runtime/Messaging/Private/Bridge/MessageBridge.cpp
+++ b/Engine/Source/Runtime/Messaging/Private/Bridge/MessageBridge.cpp
@@ -221,6 +221,7 @@ void FMessageBridge::ForgetTransportNode(const FGuid& NodeId)
 
 void FMessageBridge::ReceiveTransportMessage(const TSharedRef<IMessageContext, ESPMode::ThreadSafe>& Context, const FGuid& NodeId)
 {
+#if !NO_LOGGING
 	if (UE_GET_LOG_VERBOSITY(LogMessaging) >= ELogVerbosity::Verbose)
 	{
 		FString RecipientStr = FString::JoinBy(Context->GetRecipients(), TEXT("+"), &FMessageAddress::ToString);
@@ -229,6 +230,7 @@ void FMessageBridge::ReceiveTransportMessage(const TSharedRef<IMessageContext, E
 			*Context->GetSender().ToString(),
 			*RecipientStr);
 	}
+#endif
 
 	if (!Enabled || !Bus.IsValid())
 	{
diff --git a/Engine/Source/Runtime/Messaging/Private/Bus/MessageBus.cpp b/Engine/Source/Runtime/Messaging/Private/Bus/MessageBus.cpp
--- a/Engine/Source/Runtime/Messaging/Private/Bus/MessageBus.cpp
+++ b/Engine/Source/Runtime/Messaging/Private/Bus/MessageBus.cpp
@@ -46,6 +46,7 @@ void FMessageBus::Forward(
 	const TSharedRef<IMessageSender, ESPMode::ThreadSafe>& Forwarder
 )
 {
+#if !NO_LOGGING
 	if (UE_GET_LOG_VERBOSITY(LogMessaging) >= ELogVerbosity::Verbose)
 	{
 		FString RecipientStr = FString::JoinBy(Context->GetRecipients(), TEXT("+"), &FMessageAddress::ToString);
@@ -54,6 +55,7 @@ void FMessageBus::Forward(
 			*Context->GetMessageTypePathName().ToString(),
 			*Context->GetSender().ToString(), *RecipientStr);
 	}
+#endif
 
 	Router->RouteMessage(MakeShareable(new FMessageContext(
 		Context,
diff --git a/Engine/Source/Runtime/Messaging/Private/Bus/MessageRouter.cpp b/Engine/Source/Runtime/Messaging/Private/Bus/MessageRouter.cpp
--- a/Engine/Source/Runtime/Messaging/Private/Bus/MessageRouter.cpp
+++ b/Engine/Source/Runtime/Messaging/Private/Bus/MessageRouter.cpp
@@ -126,11 +126,14 @@ void FMessageRouter::DispatchMessage(const TSharedRef<IMessageContext, ESPMode::
 		// get recipients, either from the context...
 		if (RecipientCount > 0)
 		{
+
+#if !NO_LOGGING
 			if (UE_GET_LOG_VERBOSITY(LogMessaging) >= ELogVerbosity::Verbose)
 			{
 				FString RecipientStr = FString::JoinBy(Context->GetRecipients(), TEXT("+"), &FMessageAddress::ToString);
 				UE_LOGF(LogMessaging, Verbose, "Dispatching %ls from %ls to %ls", *Context->GetMessageTypePathName().ToString(), *Context->GetSender().ToString(), *RecipientStr);
 			}
+#endif
 
 			FilterRecipients(Context, Recipients);
 
